Adding Oculus Touchpad as an explicit input for our template project, since it no longer goes through the mobile "Touch" event.

diff --git a/Templates/TP_FirstPerson/Config/DefaultInput.ini b/Templates/TP_FirstPerson/Config/DefaultInput.ini
--- a/Templates/TP_FirstPerson/Config/DefaultInput.ini
+++ b/Templates/TP_FirstPerson/Config/DefaultInput.ini
@@ -1,31 +1,101 @@
-[/Script/Engine.InputSettings]
-+ActionMappings=(ActionName="Jump", Key=SpaceBar)
-+ActionMappings=(ActionName="Jump", Key=Gamepad_FaceButton_Bottom)
-+ActionMappings=(ActionName="Jump",Key=MotionController_Left_Trigger,bShift=False,bCtrl=False,bAlt=False,bCmd=False)
-
-+ActionMappings=(ActionName="Fire", Key=LeftMouseButton)
-+ActionMappings=(ActionName="Fire", Key=Gamepad_RightTrigger)
-+ActionMappings=(ActionName="Fire",Key=MotionController_Right_Trigger,bShift=False,bCtrl=False,bAlt=False,bCmd=False)
-
-+ActionMappings=(ActionName="ResetVR",Key=R,bShift=False,bCtrl=False,bAlt=False,bCmd=False)
-+ActionMappings=(ActionName="ResetVR",Key=MotionController_Left_Grip1,bShift=False,bCtrl=False,bAlt=False,bCmd=False)
 
-+AxisMappings=(AxisName="MoveForward", Key=W, Scale=1.f)
-+AxisMappings=(AxisName="MoveForward", Key=S, Scale=-1.f)
-+AxisMappings=(AxisName="MoveForward", Key=Up, Scale=1.f)
-+AxisMappings=(AxisName="MoveForward", Key=Down, Scale=-1.f)
-+AxisMappings=(AxisName="MoveForward", Key=Gamepad_LeftY, Scale=1.f)
-+AxisMappings=(AxisName="MoveForward",Key=MotionController_Left_Thumbstick_Y,Scale=-1.000000)
-
-+AxisMappings=(AxisName="MoveRight", Key=A, Scale=-1.f)
-+AxisMappings=(AxisName="MoveRight", Key=D, Scale=1.f)
-+AxisMappings=(AxisName="MoveRight", Key=Gamepad_LeftX, Scale=1.f)
-+AxisMappings=(AxisName="MoveRight",Key=MotionController_Left_Thumbstick_X,Scale=1.000000)
+[/Script/Engine.InputSettings]
+-AxisConfig=(AxisKeyName="Gamepad_LeftX",AxisProperties=(DeadZone=0.25,Exponent=1.f,Sensitivity=1.f))
+-AxisConfig=(AxisKeyName="Gamepad_LeftY",AxisProperties=(DeadZone=0.25,Exponent=1.f,Sensitivity=1.f))
+-AxisConfig=(AxisKeyName="Gamepad_RightX",AxisProperties=(DeadZone=0.25,Exponent=1.f,Sensitivity=1.f))
+-AxisConfig=(AxisKeyName="Gamepad_RightY",AxisProperties=(DeadZone=0.25,Exponent=1.f,Sensitivity=1.f))
+-AxisConfig=(AxisKeyName="MouseX",AxisProperties=(DeadZone=0.f,Exponent=1.f,Sensitivity=0.07f))
+-AxisConfig=(AxisKeyName="MouseY",AxisProperties=(DeadZone=0.f,Exponent=1.f,Sensitivity=0.07f))
++AxisConfig=(AxisKeyName="Gamepad_LeftX",AxisProperties=(DeadZone=0.250000,Sensitivity=1.000000,Exponent=1.000000,bInvert=False))
++AxisConfig=(AxisKeyName="Gamepad_LeftY",AxisProperties=(DeadZone=0.250000,Sensitivity=1.000000,Exponent=1.000000,bInvert=False))
++AxisConfig=(AxisKeyName="Gamepad_RightX",AxisProperties=(DeadZone=0.250000,Sensitivity=1.000000,Exponent=1.000000,bInvert=False))
++AxisConfig=(AxisKeyName="Gamepad_RightY",AxisProperties=(DeadZone=0.250000,Sensitivity=1.000000,Exponent=1.000000,bInvert=False))
++AxisConfig=(AxisKeyName="MouseX",AxisProperties=(DeadZone=0.000000,Sensitivity=0.070000,Exponent=1.000000,bInvert=False))
++AxisConfig=(AxisKeyName="MouseY",AxisProperties=(DeadZone=0.000000,Sensitivity=0.070000,Exponent=1.000000,bInvert=False))
++AxisConfig=(AxisKeyName="MotionController_Left_Thumbstick_Z",AxisProperties=(DeadZone=0.000000,Sensitivity=1.000000,Exponent=1.000000,bInvert=False))
++AxisConfig=(AxisKeyName="MotionController_Right_Thumbstick_Z",AxisProperties=(DeadZone=0.000000,Sensitivity=1.000000,Exponent=1.000000,bInvert=False))
++AxisConfig=(AxisKeyName="MouseWheelAxis",AxisProperties=(DeadZone=0.000000,Sensitivity=1.000000,Exponent=1.000000,bInvert=False))
++AxisConfig=(AxisKeyName="Gamepad_LeftTriggerAxis",AxisProperties=(DeadZone=0.000000,Sensitivity=1.000000,Exponent=1.000000,bInvert=False))
++AxisConfig=(AxisKeyName="Gamepad_RightTriggerAxis",AxisProperties=(DeadZone=0.000000,Sensitivity=1.000000,Exponent=1.000000,bInvert=False))
++AxisConfig=(AxisKeyName="MotionController_Left_Thumbstick_X",AxisProperties=(DeadZone=0.000000,Sensitivity=1.000000,Exponent=1.000000,bInvert=False))
++AxisConfig=(AxisKeyName="MotionController_Left_Thumbstick_Y",AxisProperties=(DeadZone=0.000000,Sensitivity=1.000000,Exponent=1.000000,bInvert=False))
++AxisConfig=(AxisKeyName="MotionController_Left_TriggerAxis",AxisProperties=(DeadZone=0.000000,Sensitivity=1.000000,Exponent=1.000000,bInvert=False))
++AxisConfig=(AxisKeyName="MotionController_Left_Grip1Axis",AxisProperties=(DeadZone=0.000000,Sensitivity=1.000000,Exponent=1.000000,bInvert=False))
++AxisConfig=(AxisKeyName="MotionController_Left_Grip2Axis",AxisProperties=(DeadZone=0.000000,Sensitivity=1.000000,Exponent=1.000000,bInvert=False))
++AxisConfig=(AxisKeyName="MotionController_Right_Thumbstick_X",AxisProperties=(DeadZone=0.000000,Sensitivity=1.000000,Exponent=1.000000,bInvert=False))
++AxisConfig=(AxisKeyName="MotionController_Right_Thumbstick_Y",AxisProperties=(DeadZone=0.000000,Sensitivity=1.000000,Exponent=1.000000,bInvert=False))
++AxisConfig=(AxisKeyName="MotionController_Right_TriggerAxis",AxisProperties=(DeadZone=0.000000,Sensitivity=1.000000,Exponent=1.000000,bInvert=False))
++AxisConfig=(AxisKeyName="MotionController_Right_Grip1Axis",AxisProperties=(DeadZone=0.000000,Sensitivity=1.000000,Exponent=1.000000,bInvert=False))
++AxisConfig=(AxisKeyName="MotionController_Right_Grip2Axis",AxisProperties=(DeadZone=0.000000,Sensitivity=1.000000,Exponent=1.000000,bInvert=False))
++AxisConfig=(AxisKeyName="Gamepad_Special_Left_X",AxisProperties=(DeadZone=0.000000,Sensitivity=1.000000,Exponent=1.000000,bInvert=False))
++AxisConfig=(AxisKeyName="Gamepad_Special_Left_Y",AxisProperties=(DeadZone=0.000000,Sensitivity=1.000000,Exponent=1.000000,bInvert=False))
++AxisConfig=(AxisKeyName="OculusTouch_Left_Thumbstick",AxisProperties=(DeadZone=0.000000,Sensitivity=1.000000,Exponent=1.000000,bInvert=False))
++AxisConfig=(AxisKeyName="OculusTouch_Left_FaceButton1",AxisProperties=(DeadZone=0.000000,Sensitivity=1.000000,Exponent=1.000000,bInvert=False))
++AxisConfig=(AxisKeyName="OculusTouch_Left_Trigger",AxisProperties=(DeadZone=0.000000,Sensitivity=1.000000,Exponent=1.000000,bInvert=False))
++AxisConfig=(AxisKeyName="OculusTouch_Left_FaceButton2",AxisProperties=(DeadZone=0.000000,Sensitivity=1.000000,Exponent=1.000000,bInvert=False))
++AxisConfig=(AxisKeyName="OculusTouch_Left_IndexPointing",AxisProperties=(DeadZone=0.000000,Sensitivity=1.000000,Exponent=1.000000,bInvert=False))
++AxisConfig=(AxisKeyName="OculusTouch_Left_ThumbUp",AxisProperties=(DeadZone=0.000000,Sensitivity=1.000000,Exponent=1.000000,bInvert=False))
++AxisConfig=(AxisKeyName="OculusTouch_Right_Thumbstick",AxisProperties=(DeadZone=0.000000,Sensitivity=1.000000,Exponent=1.000000,bInvert=False))
++AxisConfig=(AxisKeyName="OculusTouch_Right_FaceButton1",AxisProperties=(DeadZone=0.000000,Sensitivity=1.000000,Exponent=1.000000,bInvert=False))
++AxisConfig=(AxisKeyName="OculusTouch_Right_Trigger",AxisProperties=(DeadZone=0.000000,Sensitivity=1.000000,Exponent=1.000000,bInvert=False))
++AxisConfig=(AxisKeyName="OculusTouch_Right_FaceButton2",AxisProperties=(DeadZone=0.000000,Sensitivity=1.000000,Exponent=1.000000,bInvert=False))
++AxisConfig=(AxisKeyName="OculusTouch_Right_IndexPointing",AxisProperties=(DeadZone=0.000000,Sensitivity=1.000000,Exponent=1.000000,bInvert=False))
++AxisConfig=(AxisKeyName="OculusTouch_Right_ThumbUp",AxisProperties=(DeadZone=0.000000,Sensitivity=1.000000,Exponent=1.000000,bInvert=False))
++AxisConfig=(AxisKeyName="OculusTouchpad_Touchpad_X",AxisProperties=(DeadZone=0.000000,Sensitivity=1.000000,Exponent=1.000000,bInvert=False))
++AxisConfig=(AxisKeyName="OculusTouchpad_Touchpad_Y",AxisProperties=(DeadZone=0.000000,Sensitivity=1.000000,Exponent=1.000000,bInvert=False))
++AxisConfig=(AxisKeyName="SteamVR_Knuckles_Left_HandGrip",AxisProperties=(DeadZone=0.000000,Sensitivity=1.000000,Exponent=1.000000,bInvert=False))
++AxisConfig=(AxisKeyName="SteamVR_Knuckles_Left_IndexGrip",AxisProperties=(DeadZone=0.000000,Sensitivity=1.000000,Exponent=1.000000,bInvert=False))
++AxisConfig=(AxisKeyName="SteamVR_Knuckles_Left_MiddleGrip",AxisProperties=(DeadZone=0.000000,Sensitivity=1.000000,Exponent=1.000000,bInvert=False))
++AxisConfig=(AxisKeyName="SteamVR_Knuckles_Left_RingGrip",AxisProperties=(DeadZone=0.000000,Sensitivity=1.000000,Exponent=1.000000,bInvert=False))
++AxisConfig=(AxisKeyName="SteamVR_Knuckles_Left_PinkyGrip",AxisProperties=(DeadZone=0.000000,Sensitivity=1.000000,Exponent=1.000000,bInvert=False))
++AxisConfig=(AxisKeyName="SteamVR_Knuckles_Right_HandGrip",AxisProperties=(DeadZone=0.000000,Sensitivity=1.000000,Exponent=1.000000,bInvert=False))
++AxisConfig=(AxisKeyName="SteamVR_Knuckles_Right_IndexGrip",AxisProperties=(DeadZone=0.000000,Sensitivity=1.000000,Exponent=1.000000,bInvert=False))
++AxisConfig=(AxisKeyName="SteamVR_Knuckles_Right_MiddleGrip",AxisProperties=(DeadZone=0.000000,Sensitivity=1.000000,Exponent=1.000000,bInvert=False))
++AxisConfig=(AxisKeyName="SteamVR_Knuckles_Right_RingGrip",AxisProperties=(DeadZone=0.000000,Sensitivity=1.000000,Exponent=1.000000,bInvert=False))
++AxisConfig=(AxisKeyName="SteamVR_Knuckles_Right_PinkyGrip",AxisProperties=(DeadZone=0.000000,Sensitivity=1.000000,Exponent=1.000000,bInvert=False))
+bAltEnterTogglesFullscreen=True
+bF11TogglesFullscreen=True
+bUseMouseForTouch=False
+bEnableMouseSmoothing=True
+bEnableFOVScaling=True
+bCaptureMouseOnLaunch=True
+bDefaultViewportMouseLock=False
+bAlwaysShowTouchInterface=False
+bShowConsoleOnFourFingerTap=True
+bEnableGestureRecognizer=False
+bUseAutocorrect=False
+DefaultViewportMouseCaptureMode=CapturePermanently_IncludingInitialMouseDown
+DefaultViewportMouseLockMode=LockOnCapture
+FOVScale=0.011110
+DoubleClickTime=0.200000
++ActionMappings=(ActionName="Jump",bShift=False,bCtrl=False,bAlt=False,bCmd=False,Key=SpaceBar)
++ActionMappings=(ActionName="Jump",bShift=False,bCtrl=False,bAlt=False,bCmd=False,Key=Gamepad_FaceButton_Bottom)
++ActionMappings=(ActionName="Jump",bShift=False,bCtrl=False,bAlt=False,bCmd=False,Key=MotionController_Left_Trigger)
++ActionMappings=(ActionName="Fire",bShift=False,bCtrl=False,bAlt=False,bCmd=False,Key=LeftMouseButton)
++ActionMappings=(ActionName="Fire",bShift=False,bCtrl=False,bAlt=False,bCmd=False,Key=Gamepad_RightTrigger)
++ActionMappings=(ActionName="Fire",bShift=False,bCtrl=False,bAlt=False,bCmd=False,Key=MotionController_Right_Trigger)
++ActionMappings=(ActionName="ResetVR",bShift=False,bCtrl=False,bAlt=False,bCmd=False,Key=R)
++ActionMappings=(ActionName="ResetVR",bShift=False,bCtrl=False,bAlt=False,bCmd=False,Key=MotionController_Left_Grip1)
++ActionMappings=(ActionName="Fire",bShift=False,bCtrl=False,bAlt=False,bCmd=False,Key=OculusTouchpad_Touchpad)
++AxisMappings=(AxisName="MoveForward",Scale=1.000000,Key=W)
++AxisMappings=(AxisName="MoveForward",Scale=-1.000000,Key=S)
++AxisMappings=(AxisName="MoveForward",Scale=1.000000,Key=Up)
++AxisMappings=(AxisName="MoveForward",Scale=-1.000000,Key=Down)
++AxisMappings=(AxisName="MoveForward",Scale=1.000000,Key=Gamepad_LeftY)
++AxisMappings=(AxisName="MoveForward",Scale=-1.000000,Key=MotionController_Left_Thumbstick_Y)
++AxisMappings=(AxisName="MoveRight",Scale=-1.000000,Key=A)
++AxisMappings=(AxisName="MoveRight",Scale=1.000000,Key=D)
++AxisMappings=(AxisName="MoveRight",Scale=1.000000,Key=Gamepad_LeftX)
++AxisMappings=(AxisName="MoveRight",Scale=1.000000,Key=MotionController_Left_Thumbstick_X)
++AxisMappings=(AxisName="TurnRate",Scale=1.000000,Key=Gamepad_RightX)
++AxisMappings=(AxisName="TurnRate",Scale=-1.000000,Key=Left)
++AxisMappings=(AxisName="TurnRate",Scale=1.000000,Key=Right)
++AxisMappings=(AxisName="Turn",Scale=1.000000,Key=MouseX)
++AxisMappings=(AxisName="LookUpRate",Scale=1.000000,Key=Gamepad_RightY)
++AxisMappings=(AxisName="LookUp",Scale=-1.000000,Key=MouseY)
+DefaultTouchInterface=/Engine/MobileResources/HUD/DefaultVirtualJoysticks.DefaultVirtualJoysticks
+ConsoleKey=None
+-ConsoleKeys=Tilde
++ConsoleKeys=Tilde
 
-+AxisMappings=(AxisName="TurnRate", Key=Gamepad_RightX, Scale=1.f)
-+AxisMappings=(AxisName="TurnRate", Key=Left, Scale=-1.f)
-+AxisMappings=(AxisName="TurnRate", Key=Right, Scale=1.f)
-+AxisMappings=(AxisName="Turn", Key=MouseX, Scale=1.f)
 
-+AxisMappings=(AxisName="LookUpRate", Key=Gamepad_RightY, Scale=1.f)
-+AxisMappings=(AxisName="LookUp", Key=MouseY, Scale=-1.f)
\ No newline at end of file
diff --git a/Templates/TP_FirstPersonBP/Config/DefaultInput.ini b/Templates/TP_FirstPersonBP/Config/DefaultInput.ini
--- a/Templates/TP_FirstPersonBP/Config/DefaultInput.ini
+++ b/Templates/TP_FirstPersonBP/Config/DefaultInput.ini
@@ -27,41 +27,72 @@
 +AxisConfig=(AxisKeyName="MotionController_Right_Grip2Axis",AxisProperties=(DeadZone=0.000000,Sensitivity=1.000000,Exponent=1.000000,bInvert=False))
 +AxisConfig=(AxisKeyName="Gamepad_Special_Left_X",AxisProperties=(DeadZone=0.000000,Sensitivity=1.000000,Exponent=1.000000,bInvert=False))
 +AxisConfig=(AxisKeyName="Gamepad_Special_Left_Y",AxisProperties=(DeadZone=0.000000,Sensitivity=1.000000,Exponent=1.000000,bInvert=False))
++AxisConfig=(AxisKeyName="MotionController_Left_Thumbstick_Z",AxisProperties=(DeadZone=0.000000,Sensitivity=1.000000,Exponent=1.000000,bInvert=False))
++AxisConfig=(AxisKeyName="MotionController_Right_Thumbstick_Z",AxisProperties=(DeadZone=0.000000,Sensitivity=1.000000,Exponent=1.000000,bInvert=False))
++AxisConfig=(AxisKeyName="OculusTouch_Left_Thumbstick",AxisProperties=(DeadZone=0.000000,Sensitivity=1.000000,Exponent=1.000000,bInvert=False))
++AxisConfig=(AxisKeyName="OculusTouch_Left_FaceButton1",AxisProperties=(DeadZone=0.000000,Sensitivity=1.000000,Exponent=1.000000,bInvert=False))
++AxisConfig=(AxisKeyName="OculusTouch_Left_Trigger",AxisProperties=(DeadZone=0.000000,Sensitivity=1.000000,Exponent=1.000000,bInvert=False))
++AxisConfig=(AxisKeyName="OculusTouch_Left_FaceButton2",AxisProperties=(DeadZone=0.000000,Sensitivity=1.000000,Exponent=1.000000,bInvert=False))
++AxisConfig=(AxisKeyName="OculusTouch_Left_IndexPointing",AxisProperties=(DeadZone=0.000000,Sensitivity=1.000000,Exponent=1.000000,bInvert=False))
++AxisConfig=(AxisKeyName="OculusTouch_Left_ThumbUp",AxisProperties=(DeadZone=0.000000,Sensitivity=1.000000,Exponent=1.000000,bInvert=False))
++AxisConfig=(AxisKeyName="OculusTouch_Right_Thumbstick",AxisProperties=(DeadZone=0.000000,Sensitivity=1.000000,Exponent=1.000000,bInvert=False))
++AxisConfig=(AxisKeyName="OculusTouch_Right_FaceButton1",AxisProperties=(DeadZone=0.000000,Sensitivity=1.000000,Exponent=1.000000,bInvert=False))
++AxisConfig=(AxisKeyName="OculusTouch_Right_Trigger",AxisProperties=(DeadZone=0.000000,Sensitivity=1.000000,Exponent=1.000000,bInvert=False))
++AxisConfig=(AxisKeyName="OculusTouch_Right_FaceButton2",AxisProperties=(DeadZone=0.000000,Sensitivity=1.000000,Exponent=1.000000,bInvert=False))
++AxisConfig=(AxisKeyName="OculusTouch_Right_IndexPointing",AxisProperties=(DeadZone=0.000000,Sensitivity=1.000000,Exponent=1.000000,bInvert=False))
++AxisConfig=(AxisKeyName="OculusTouch_Right_ThumbUp",AxisProperties=(DeadZone=0.000000,Sensitivity=1.000000,Exponent=1.000000,bInvert=False))
++AxisConfig=(AxisKeyName="OculusTouchpad_Touchpad_X",AxisProperties=(DeadZone=0.000000,Sensitivity=1.000000,Exponent=1.000000,bInvert=False))
++AxisConfig=(AxisKeyName="OculusTouchpad_Touchpad_Y",AxisProperties=(DeadZone=0.000000,Sensitivity=1.000000,Exponent=1.000000,bInvert=False))
++AxisConfig=(AxisKeyName="SteamVR_Knuckles_Left_HandGrip",AxisProperties=(DeadZone=0.000000,Sensitivity=1.000000,Exponent=1.000000,bInvert=False))
++AxisConfig=(AxisKeyName="SteamVR_Knuckles_Left_IndexGrip",AxisProperties=(DeadZone=0.000000,Sensitivity=1.000000,Exponent=1.000000,bInvert=False))
++AxisConfig=(AxisKeyName="SteamVR_Knuckles_Left_MiddleGrip",AxisProperties=(DeadZone=0.000000,Sensitivity=1.000000,Exponent=1.000000,bInvert=False))
++AxisConfig=(AxisKeyName="SteamVR_Knuckles_Left_RingGrip",AxisProperties=(DeadZone=0.000000,Sensitivity=1.000000,Exponent=1.000000,bInvert=False))
++AxisConfig=(AxisKeyName="SteamVR_Knuckles_Left_PinkyGrip",AxisProperties=(DeadZone=0.000000,Sensitivity=1.000000,Exponent=1.000000,bInvert=False))
++AxisConfig=(AxisKeyName="SteamVR_Knuckles_Right_HandGrip",AxisProperties=(DeadZone=0.000000,Sensitivity=1.000000,Exponent=1.000000,bInvert=False))
++AxisConfig=(AxisKeyName="SteamVR_Knuckles_Right_IndexGrip",AxisProperties=(DeadZone=0.000000,Sensitivity=1.000000,Exponent=1.000000,bInvert=False))
++AxisConfig=(AxisKeyName="SteamVR_Knuckles_Right_MiddleGrip",AxisProperties=(DeadZone=0.000000,Sensitivity=1.000000,Exponent=1.000000,bInvert=False))
++AxisConfig=(AxisKeyName="SteamVR_Knuckles_Right_RingGrip",AxisProperties=(DeadZone=0.000000,Sensitivity=1.000000,Exponent=1.000000,bInvert=False))
++AxisConfig=(AxisKeyName="SteamVR_Knuckles_Right_PinkyGrip",AxisProperties=(DeadZone=0.000000,Sensitivity=1.000000,Exponent=1.000000,bInvert=False))
 bAltEnterTogglesFullscreen=True
+bF11TogglesFullscreen=True
 bUseMouseForTouch=False
 bEnableMouseSmoothing=True
 bEnableFOVScaling=True
-FOVScale=0.011110
-DoubleClickTime=0.200000
 bCaptureMouseOnLaunch=True
-DefaultViewportMouseCaptureMode=CapturePermanently_IncludingInitialMouseDown
 bDefaultViewportMouseLock=True
-+ActionMappings=(ActionName="Jump",Key=SpaceBar,bShift=False,bCtrl=False,bAlt=False,bCmd=False)
-+ActionMappings=(ActionName="Jump",Key=Gamepad_FaceButton_Bottom,bShift=False,bCtrl=False,bAlt=False,bCmd=False)
-+ActionMappings=(ActionName="Fire",Key=LeftMouseButton,bShift=False,bCtrl=False,bAlt=False,bCmd=False)
-+ActionMappings=(ActionName="Fire",Key=Gamepad_RightTrigger,bShift=False,bCtrl=False,bAlt=False,bCmd=False)
-+ActionMappings=(ActionName="Fire",Key=MotionController_Right_Trigger,bShift=False,bCtrl=False,bAlt=False,bCmd=False)
-+ActionMappings=(ActionName="Jump",Key=MotionController_Left_Trigger,bShift=False,bCtrl=False,bAlt=False,bCmd=False)
-+ActionMappings=(ActionName="ResetVR",Key=R,bShift=False,bCtrl=False,bAlt=False,bCmd=False)
-+ActionMappings=(ActionName="ResetVR",Key=MotionController_Left_Grip1,bShift=False,bCtrl=False,bAlt=False,bCmd=False)
-+AxisMappings=(AxisName="MoveForward",Key=W,Scale=1.000000)
-+AxisMappings=(AxisName="MoveForward",Key=S,Scale=-1.000000)
-+AxisMappings=(AxisName="MoveForward",Key=Up,Scale=1.000000)
-+AxisMappings=(AxisName="MoveForward",Key=Down,Scale=-1.000000)
-+AxisMappings=(AxisName="MoveForward",Key=Gamepad_LeftY,Scale=1.000000)
-+AxisMappings=(AxisName="MoveRight",Key=A,Scale=-1.000000)
-+AxisMappings=(AxisName="MoveRight",Key=D,Scale=1.000000)
-+AxisMappings=(AxisName="MoveRight",Key=Gamepad_LeftX,Scale=1.000000)
-+AxisMappings=(AxisName="TurnRate",Key=Gamepad_RightX,Scale=1.000000)
-+AxisMappings=(AxisName="TurnRate",Key=Left,Scale=-1.000000)
-+AxisMappings=(AxisName="TurnRate",Key=Right,Scale=1.000000)
-+AxisMappings=(AxisName="Turn",Key=MouseX,Scale=1.000000)
-+AxisMappings=(AxisName="LookUpRate",Key=Gamepad_RightY,Scale=1.000000)
-+AxisMappings=(AxisName="LookUp",Key=MouseY,Scale=-1.000000)
-+AxisMappings=(AxisName="MoveForward",Key=MotionController_Left_Thumbstick_Y,Scale=-1.000000)
-+AxisMappings=(AxisName="MoveRight",Key=MotionController_Left_Thumbstick_X,Scale=1.000000)
 bAlwaysShowTouchInterface=False
 bShowConsoleOnFourFingerTap=True
+bEnableGestureRecognizer=False
+bUseAutocorrect=False
+DefaultViewportMouseCaptureMode=CapturePermanently_IncludingInitialMouseDown
+DefaultViewportMouseLockMode=LockOnCapture
+FOVScale=0.011110
+DoubleClickTime=0.200000
++ActionMappings=(ActionName="Jump",bShift=False,bCtrl=False,bAlt=False,bCmd=False,Key=SpaceBar)
++ActionMappings=(ActionName="Jump",bShift=False,bCtrl=False,bAlt=False,bCmd=False,Key=Gamepad_FaceButton_Bottom)
++ActionMappings=(ActionName="Fire",bShift=False,bCtrl=False,bAlt=False,bCmd=False,Key=LeftMouseButton)
++ActionMappings=(ActionName="Fire",bShift=False,bCtrl=False,bAlt=False,bCmd=False,Key=Gamepad_RightTrigger)
++ActionMappings=(ActionName="Fire",bShift=False,bCtrl=False,bAlt=False,bCmd=False,Key=MotionController_Right_Trigger)
++ActionMappings=(ActionName="Jump",bShift=False,bCtrl=False,bAlt=False,bCmd=False,Key=MotionController_Left_Trigger)
++ActionMappings=(ActionName="ResetVR",bShift=False,bCtrl=False,bAlt=False,bCmd=False,Key=R)
++ActionMappings=(ActionName="ResetVR",bShift=False,bCtrl=False,bAlt=False,bCmd=False,Key=MotionController_Left_Grip1)
++ActionMappings=(ActionName="Fire",bShift=False,bCtrl=False,bAlt=False,bCmd=False,Key=OculusTouchpad_Touchpad)
++AxisMappings=(AxisName="MoveForward",Scale=1.000000,Key=W)
++AxisMappings=(AxisName="MoveForward",Scale=-1.000000,Key=S)
++AxisMappings=(AxisName="MoveForward",Scale=1.000000,Key=Up)
++AxisMappings=(AxisName="MoveForward",Scale=-1.000000,Key=Down)
++AxisMappings=(AxisName="MoveForward",Scale=1.000000,Key=Gamepad_LeftY)
++AxisMappings=(AxisName="MoveRight",Scale=-1.000000,Key=A)
++AxisMappings=(AxisName="MoveRight",Scale=1.000000,Key=D)
++AxisMappings=(AxisName="MoveRight",Scale=1.000000,Key=Gamepad_LeftX)
++AxisMappings=(AxisName="TurnRate",Scale=1.000000,Key=Gamepad_RightX)
++AxisMappings=(AxisName="TurnRate",Scale=-1.000000,Key=Left)
++AxisMappings=(AxisName="TurnRate",Scale=1.000000,Key=Right)
++AxisMappings=(AxisName="Turn",Scale=1.000000,Key=MouseX)
++AxisMappings=(AxisName="LookUpRate",Scale=1.000000,Key=Gamepad_RightY)
++AxisMappings=(AxisName="LookUp",Scale=-1.000000,Key=MouseY)
++AxisMappings=(AxisName="MoveForward",Scale=-1.000000,Key=MotionController_Left_Thumbstick_Y)
++AxisMappings=(AxisName="MoveRight",Scale=1.000000,Key=MotionController_Left_Thumbstick_X)
 DefaultTouchInterface=/Engine/MobileResources/HUD/DefaultVirtualJoysticks.DefaultVirtualJoysticks
 ConsoleKey=None
 -ConsoleKeys=Tilde
diff --git a/Templates/TP_ThirdPerson/Config/DefaultInput.ini b/Templates/TP_ThirdPerson/Config/DefaultInput.ini
--- a/Templates/TP_ThirdPerson/Config/DefaultInput.ini
+++ b/Templates/TP_ThirdPerson/Config/DefaultInput.ini
@@ -1,24 +1,95 @@
-[/Script/Engine.InputSettings]
-+ActionMappings=(ActionName="Jump", Key=SpaceBar)
-+ActionMappings=(ActionName="Jump", Key=Gamepad_FaceButton_Bottom)
-
-+AxisMappings=(AxisName="MoveForward", Key=W, Scale=1.f)
-+AxisMappings=(AxisName="MoveForward", Key=S, Scale=-1.f)
-+AxisMappings=(AxisName="MoveForward", Key=Up, Scale=1.f)
-+AxisMappings=(AxisName="MoveForward", Key=Down, Scale=-1.f)
-+AxisMappings=(AxisName="MoveForward", Key=Gamepad_LeftY, Scale=1.f)
 
-+AxisMappings=(AxisName="MoveRight", Key=A, Scale=-1.f)
-+AxisMappings=(AxisName="MoveRight", Key=D, Scale=1.f)
-+AxisMappings=(AxisName="MoveRight", Key=Gamepad_LeftX, Scale=1.f)
-
-+AxisMappings=(AxisName="TurnRate", Key=Gamepad_RightX, Scale=1.f)
-+AxisMappings=(AxisName="TurnRate", Key=Left, Scale=-1.f)
-+AxisMappings=(AxisName="TurnRate", Key=Right, Scale=1.f)
-+AxisMappings=(AxisName="Turn", Key=MouseX, Scale=1.f)
+[/Script/Engine.InputSettings]
+-AxisConfig=(AxisKeyName="Gamepad_LeftX",AxisProperties=(DeadZone=0.25,Exponent=1.f,Sensitivity=1.f))
+-AxisConfig=(AxisKeyName="Gamepad_LeftY",AxisProperties=(DeadZone=0.25,Exponent=1.f,Sensitivity=1.f))
+-AxisConfig=(AxisKeyName="Gamepad_RightX",AxisProperties=(DeadZone=0.25,Exponent=1.f,Sensitivity=1.f))
+-AxisConfig=(AxisKeyName="Gamepad_RightY",AxisProperties=(DeadZone=0.25,Exponent=1.f,Sensitivity=1.f))
+-AxisConfig=(AxisKeyName="MouseX",AxisProperties=(DeadZone=0.f,Exponent=1.f,Sensitivity=0.07f))
+-AxisConfig=(AxisKeyName="MouseY",AxisProperties=(DeadZone=0.f,Exponent=1.f,Sensitivity=0.07f))
++AxisConfig=(AxisKeyName="Gamepad_LeftX",AxisProperties=(DeadZone=0.250000,Sensitivity=1.000000,Exponent=1.000000,bInvert=False))
++AxisConfig=(AxisKeyName="Gamepad_LeftY",AxisProperties=(DeadZone=0.250000,Sensitivity=1.000000,Exponent=1.000000,bInvert=False))
++AxisConfig=(AxisKeyName="Gamepad_RightX",AxisProperties=(DeadZone=0.250000,Sensitivity=1.000000,Exponent=1.000000,bInvert=False))
++AxisConfig=(AxisKeyName="Gamepad_RightY",AxisProperties=(DeadZone=0.250000,Sensitivity=1.000000,Exponent=1.000000,bInvert=False))
++AxisConfig=(AxisKeyName="MouseX",AxisProperties=(DeadZone=0.000000,Sensitivity=0.070000,Exponent=1.000000,bInvert=False))
++AxisConfig=(AxisKeyName="MouseY",AxisProperties=(DeadZone=0.000000,Sensitivity=0.070000,Exponent=1.000000,bInvert=False))
++AxisConfig=(AxisKeyName="MotionController_Left_Thumbstick_Z",AxisProperties=(DeadZone=0.000000,Sensitivity=1.000000,Exponent=1.000000,bInvert=False))
++AxisConfig=(AxisKeyName="MotionController_Right_Thumbstick_Z",AxisProperties=(DeadZone=0.000000,Sensitivity=1.000000,Exponent=1.000000,bInvert=False))
++AxisConfig=(AxisKeyName="MouseWheelAxis",AxisProperties=(DeadZone=0.000000,Sensitivity=1.000000,Exponent=1.000000,bInvert=False))
++AxisConfig=(AxisKeyName="Gamepad_LeftTriggerAxis",AxisProperties=(DeadZone=0.000000,Sensitivity=1.000000,Exponent=1.000000,bInvert=False))
++AxisConfig=(AxisKeyName="Gamepad_RightTriggerAxis",AxisProperties=(DeadZone=0.000000,Sensitivity=1.000000,Exponent=1.000000,bInvert=False))
++AxisConfig=(AxisKeyName="MotionController_Left_Thumbstick_X",AxisProperties=(DeadZone=0.000000,Sensitivity=1.000000,Exponent=1.000000,bInvert=False))
++AxisConfig=(AxisKeyName="MotionController_Left_Thumbstick_Y",AxisProperties=(DeadZone=0.000000,Sensitivity=1.000000,Exponent=1.000000,bInvert=False))
++AxisConfig=(AxisKeyName="MotionController_Left_TriggerAxis",AxisProperties=(DeadZone=0.000000,Sensitivity=1.000000,Exponent=1.000000,bInvert=False))
++AxisConfig=(AxisKeyName="MotionController_Left_Grip1Axis",AxisProperties=(DeadZone=0.000000,Sensitivity=1.000000,Exponent=1.000000,bInvert=False))
++AxisConfig=(AxisKeyName="MotionController_Left_Grip2Axis",AxisProperties=(DeadZone=0.000000,Sensitivity=1.000000,Exponent=1.000000,bInvert=False))
++AxisConfig=(AxisKeyName="MotionController_Right_Thumbstick_X",AxisProperties=(DeadZone=0.000000,Sensitivity=1.000000,Exponent=1.000000,bInvert=False))
++AxisConfig=(AxisKeyName="MotionController_Right_Thumbstick_Y",AxisProperties=(DeadZone=0.000000,Sensitivity=1.000000,Exponent=1.000000,bInvert=False))
++AxisConfig=(AxisKeyName="MotionController_Right_TriggerAxis",AxisProperties=(DeadZone=0.000000,Sensitivity=1.000000,Exponent=1.000000,bInvert=False))
++AxisConfig=(AxisKeyName="MotionController_Right_Grip1Axis",AxisProperties=(DeadZone=0.000000,Sensitivity=1.000000,Exponent=1.000000,bInvert=False))
++AxisConfig=(AxisKeyName="MotionController_Right_Grip2Axis",AxisProperties=(DeadZone=0.000000,Sensitivity=1.000000,Exponent=1.000000,bInvert=False))
++AxisConfig=(AxisKeyName="Gamepad_Special_Left_X",AxisProperties=(DeadZone=0.000000,Sensitivity=1.000000,Exponent=1.000000,bInvert=False))
++AxisConfig=(AxisKeyName="Gamepad_Special_Left_Y",AxisProperties=(DeadZone=0.000000,Sensitivity=1.000000,Exponent=1.000000,bInvert=False))
++AxisConfig=(AxisKeyName="OculusTouch_Left_Thumbstick",AxisProperties=(DeadZone=0.000000,Sensitivity=1.000000,Exponent=1.000000,bInvert=False))
++AxisConfig=(AxisKeyName="OculusTouch_Left_FaceButton1",AxisProperties=(DeadZone=0.000000,Sensitivity=1.000000,Exponent=1.000000,bInvert=False))
++AxisConfig=(AxisKeyName="OculusTouch_Left_Trigger",AxisProperties=(DeadZone=0.000000,Sensitivity=1.000000,Exponent=1.000000,bInvert=False))
++AxisConfig=(AxisKeyName="OculusTouch_Left_FaceButton2",AxisProperties=(DeadZone=0.000000,Sensitivity=1.000000,Exponent=1.000000,bInvert=False))
++AxisConfig=(AxisKeyName="OculusTouch_Left_IndexPointing",AxisProperties=(DeadZone=0.000000,Sensitivity=1.000000,Exponent=1.000000,bInvert=False))
++AxisConfig=(AxisKeyName="OculusTouch_Left_ThumbUp",AxisProperties=(DeadZone=0.000000,Sensitivity=1.000000,Exponent=1.000000,bInvert=False))
++AxisConfig=(AxisKeyName="OculusTouch_Right_Thumbstick",AxisProperties=(DeadZone=0.000000,Sensitivity=1.000000,Exponent=1.000000,bInvert=False))
++AxisConfig=(AxisKeyName="OculusTouch_Right_FaceButton1",AxisProperties=(DeadZone=0.000000,Sensitivity=1.000000,Exponent=1.000000,bInvert=False))
++AxisConfig=(AxisKeyName="OculusTouch_Right_Trigger",AxisProperties=(DeadZone=0.000000,Sensitivity=1.000000,Exponent=1.000000,bInvert=False))
++AxisConfig=(AxisKeyName="OculusTouch_Right_FaceButton2",AxisProperties=(DeadZone=0.000000,Sensitivity=1.000000,Exponent=1.000000,bInvert=False))
++AxisConfig=(AxisKeyName="OculusTouch_Right_IndexPointing",AxisProperties=(DeadZone=0.000000,Sensitivity=1.000000,Exponent=1.000000,bInvert=False))
++AxisConfig=(AxisKeyName="OculusTouch_Right_ThumbUp",AxisProperties=(DeadZone=0.000000,Sensitivity=1.000000,Exponent=1.000000,bInvert=False))
++AxisConfig=(AxisKeyName="OculusTouchpad_Touchpad_X",AxisProperties=(DeadZone=0.000000,Sensitivity=1.000000,Exponent=1.000000,bInvert=False))
++AxisConfig=(AxisKeyName="OculusTouchpad_Touchpad_Y",AxisProperties=(DeadZone=0.000000,Sensitivity=1.000000,Exponent=1.000000,bInvert=False))
++AxisConfig=(AxisKeyName="SteamVR_Knuckles_Left_HandGrip",AxisProperties=(DeadZone=0.000000,Sensitivity=1.000000,Exponent=1.000000,bInvert=False))
++AxisConfig=(AxisKeyName="SteamVR_Knuckles_Left_IndexGrip",AxisProperties=(DeadZone=0.000000,Sensitivity=1.000000,Exponent=1.000000,bInvert=False))
++AxisConfig=(AxisKeyName="SteamVR_Knuckles_Left_MiddleGrip",AxisProperties=(DeadZone=0.000000,Sensitivity=1.000000,Exponent=1.000000,bInvert=False))
++AxisConfig=(AxisKeyName="SteamVR_Knuckles_Left_RingGrip",AxisProperties=(DeadZone=0.000000,Sensitivity=1.000000,Exponent=1.000000,bInvert=False))
++AxisConfig=(AxisKeyName="SteamVR_Knuckles_Left_PinkyGrip",AxisProperties=(DeadZone=0.000000,Sensitivity=1.000000,Exponent=1.000000,bInvert=False))
++AxisConfig=(AxisKeyName="SteamVR_Knuckles_Right_HandGrip",AxisProperties=(DeadZone=0.000000,Sensitivity=1.000000,Exponent=1.000000,bInvert=False))
++AxisConfig=(AxisKeyName="SteamVR_Knuckles_Right_IndexGrip",AxisProperties=(DeadZone=0.000000,Sensitivity=1.000000,Exponent=1.000000,bInvert=False))
++AxisConfig=(AxisKeyName="SteamVR_Knuckles_Right_MiddleGrip",AxisProperties=(DeadZone=0.000000,Sensitivity=1.000000,Exponent=1.000000,bInvert=False))
++AxisConfig=(AxisKeyName="SteamVR_Knuckles_Right_RingGrip",AxisProperties=(DeadZone=0.000000,Sensitivity=1.000000,Exponent=1.000000,bInvert=False))
++AxisConfig=(AxisKeyName="SteamVR_Knuckles_Right_PinkyGrip",AxisProperties=(DeadZone=0.000000,Sensitivity=1.000000,Exponent=1.000000,bInvert=False))
+bAltEnterTogglesFullscreen=True
+bF11TogglesFullscreen=True
+bUseMouseForTouch=False
+bEnableMouseSmoothing=True
+bEnableFOVScaling=True
+bCaptureMouseOnLaunch=True
+bDefaultViewportMouseLock=False
+bAlwaysShowTouchInterface=False
+bShowConsoleOnFourFingerTap=True
+bEnableGestureRecognizer=False
+bUseAutocorrect=False
+DefaultViewportMouseCaptureMode=CapturePermanently_IncludingInitialMouseDown
+DefaultViewportMouseLockMode=LockOnCapture
+FOVScale=0.011110
+DoubleClickTime=0.200000
++ActionMappings=(ActionName="Jump",bShift=False,bCtrl=False,bAlt=False,bCmd=False,Key=SpaceBar)
++ActionMappings=(ActionName="Jump",bShift=False,bCtrl=False,bAlt=False,bCmd=False,Key=Gamepad_FaceButton_Bottom)
++ActionMappings=(ActionName="ResetVR",bShift=False,bCtrl=False,bAlt=False,bCmd=False,Key=R)
++ActionMappings=(ActionName="ResetVR",bShift=False,bCtrl=False,bAlt=False,bCmd=False,Key=MotionController_Left_Grip1)
++ActionMappings=(ActionName="Jump",bShift=False,bCtrl=False,bAlt=False,bCmd=False,Key=OculusTouchpad_Touchpad)
++AxisMappings=(AxisName="MoveForward",Scale=1.000000,Key=W)
++AxisMappings=(AxisName="MoveForward",Scale=-1.000000,Key=S)
++AxisMappings=(AxisName="MoveForward",Scale=1.000000,Key=Up)
++AxisMappings=(AxisName="MoveForward",Scale=-1.000000,Key=Down)
++AxisMappings=(AxisName="MoveForward",Scale=1.000000,Key=Gamepad_LeftY)
++AxisMappings=(AxisName="MoveRight",Scale=-1.000000,Key=A)
++AxisMappings=(AxisName="MoveRight",Scale=1.000000,Key=D)
++AxisMappings=(AxisName="MoveRight",Scale=1.000000,Key=Gamepad_LeftX)
++AxisMappings=(AxisName="TurnRate",Scale=1.000000,Key=Gamepad_RightX)
++AxisMappings=(AxisName="TurnRate",Scale=-1.000000,Key=Left)
++AxisMappings=(AxisName="TurnRate",Scale=1.000000,Key=Right)
++AxisMappings=(AxisName="Turn",Scale=1.000000,Key=MouseX)
++AxisMappings=(AxisName="LookUpRate",Scale=1.000000,Key=Gamepad_RightY)
++AxisMappings=(AxisName="LookUp",Scale=-1.000000,Key=MouseY)
+DefaultTouchInterface=/Engine/MobileResources/HUD/DefaultVirtualJoysticks.DefaultVirtualJoysticks
+ConsoleKey=None
+-ConsoleKeys=Tilde
++ConsoleKeys=Tilde
 
-+AxisMappings=(AxisName="LookUpRate", Key=Gamepad_RightY, Scale=1.f)
-+AxisMappings=(AxisName="LookUp", Key=MouseY, Scale=-1.f)
 
-+ActionMappings=(ActionName="ResetVR", Key=R)
-+ActionMappings=(ActionName="ResetVR",Key=MotionController_Left_Grip1,bShift=False,bCtrl=False,bAlt=False,bCmd=False)
diff --git a/Templates/TP_ThirdPersonBP/Config/DefaultInput.ini b/Templates/TP_ThirdPersonBP/Config/DefaultInput.ini
--- a/Templates/TP_ThirdPersonBP/Config/DefaultInput.ini
+++ b/Templates/TP_ThirdPersonBP/Config/DefaultInput.ini
@@ -27,40 +27,71 @@
 +AxisConfig=(AxisKeyName="MotionController_Right_Grip2Axis",AxisProperties=(DeadZone=0.000000,Sensitivity=1.000000,Exponent=1.000000,bInvert=False))
 +AxisConfig=(AxisKeyName="Gamepad_Special_Left_X",AxisProperties=(DeadZone=0.000000,Sensitivity=1.000000,Exponent=1.000000,bInvert=False))
 +AxisConfig=(AxisKeyName="Gamepad_Special_Left_Y",AxisProperties=(DeadZone=0.000000,Sensitivity=1.000000,Exponent=1.000000,bInvert=False))
++AxisConfig=(AxisKeyName="MotionController_Left_Thumbstick_Z",AxisProperties=(DeadZone=0.000000,Sensitivity=1.000000,Exponent=1.000000,bInvert=False))
++AxisConfig=(AxisKeyName="MotionController_Right_Thumbstick_Z",AxisProperties=(DeadZone=0.000000,Sensitivity=1.000000,Exponent=1.000000,bInvert=False))
++AxisConfig=(AxisKeyName="OculusTouch_Left_Thumbstick",AxisProperties=(DeadZone=0.000000,Sensitivity=1.000000,Exponent=1.000000,bInvert=False))
++AxisConfig=(AxisKeyName="OculusTouch_Left_FaceButton1",AxisProperties=(DeadZone=0.000000,Sensitivity=1.000000,Exponent=1.000000,bInvert=False))
++AxisConfig=(AxisKeyName="OculusTouch_Left_Trigger",AxisProperties=(DeadZone=0.000000,Sensitivity=1.000000,Exponent=1.000000,bInvert=False))
++AxisConfig=(AxisKeyName="OculusTouch_Left_FaceButton2",AxisProperties=(DeadZone=0.000000,Sensitivity=1.000000,Exponent=1.000000,bInvert=False))
++AxisConfig=(AxisKeyName="OculusTouch_Left_IndexPointing",AxisProperties=(DeadZone=0.000000,Sensitivity=1.000000,Exponent=1.000000,bInvert=False))
++AxisConfig=(AxisKeyName="OculusTouch_Left_ThumbUp",AxisProperties=(DeadZone=0.000000,Sensitivity=1.000000,Exponent=1.000000,bInvert=False))
++AxisConfig=(AxisKeyName="OculusTouch_Right_Thumbstick",AxisProperties=(DeadZone=0.000000,Sensitivity=1.000000,Exponent=1.000000,bInvert=False))
++AxisConfig=(AxisKeyName="OculusTouch_Right_FaceButton1",AxisProperties=(DeadZone=0.000000,Sensitivity=1.000000,Exponent=1.000000,bInvert=False))
++AxisConfig=(AxisKeyName="OculusTouch_Right_Trigger",AxisProperties=(DeadZone=0.000000,Sensitivity=1.000000,Exponent=1.000000,bInvert=False))
++AxisConfig=(AxisKeyName="OculusTouch_Right_FaceButton2",AxisProperties=(DeadZone=0.000000,Sensitivity=1.000000,Exponent=1.000000,bInvert=False))
++AxisConfig=(AxisKeyName="OculusTouch_Right_IndexPointing",AxisProperties=(DeadZone=0.000000,Sensitivity=1.000000,Exponent=1.000000,bInvert=False))
++AxisConfig=(AxisKeyName="OculusTouch_Right_ThumbUp",AxisProperties=(DeadZone=0.000000,Sensitivity=1.000000,Exponent=1.000000,bInvert=False))
++AxisConfig=(AxisKeyName="OculusTouchpad_Touchpad_X",AxisProperties=(DeadZone=0.000000,Sensitivity=1.000000,Exponent=1.000000,bInvert=False))
++AxisConfig=(AxisKeyName="OculusTouchpad_Touchpad_Y",AxisProperties=(DeadZone=0.000000,Sensitivity=1.000000,Exponent=1.000000,bInvert=False))
++AxisConfig=(AxisKeyName="SteamVR_Knuckles_Left_HandGrip",AxisProperties=(DeadZone=0.000000,Sensitivity=1.000000,Exponent=1.000000,bInvert=False))
++AxisConfig=(AxisKeyName="SteamVR_Knuckles_Left_IndexGrip",AxisProperties=(DeadZone=0.000000,Sensitivity=1.000000,Exponent=1.000000,bInvert=False))
++AxisConfig=(AxisKeyName="SteamVR_Knuckles_Left_MiddleGrip",AxisProperties=(DeadZone=0.000000,Sensitivity=1.000000,Exponent=1.000000,bInvert=False))
++AxisConfig=(AxisKeyName="SteamVR_Knuckles_Left_RingGrip",AxisProperties=(DeadZone=0.000000,Sensitivity=1.000000,Exponent=1.000000,bInvert=False))
++AxisConfig=(AxisKeyName="SteamVR_Knuckles_Left_PinkyGrip",AxisProperties=(DeadZone=0.000000,Sensitivity=1.000000,Exponent=1.000000,bInvert=False))
++AxisConfig=(AxisKeyName="SteamVR_Knuckles_Right_HandGrip",AxisProperties=(DeadZone=0.000000,Sensitivity=1.000000,Exponent=1.000000,bInvert=False))
++AxisConfig=(AxisKeyName="SteamVR_Knuckles_Right_IndexGrip",AxisProperties=(DeadZone=0.000000,Sensitivity=1.000000,Exponent=1.000000,bInvert=False))
++AxisConfig=(AxisKeyName="SteamVR_Knuckles_Right_MiddleGrip",AxisProperties=(DeadZone=0.000000,Sensitivity=1.000000,Exponent=1.000000,bInvert=False))
++AxisConfig=(AxisKeyName="SteamVR_Knuckles_Right_RingGrip",AxisProperties=(DeadZone=0.000000,Sensitivity=1.000000,Exponent=1.000000,bInvert=False))
++AxisConfig=(AxisKeyName="SteamVR_Knuckles_Right_PinkyGrip",AxisProperties=(DeadZone=0.000000,Sensitivity=1.000000,Exponent=1.000000,bInvert=False))
 bAltEnterTogglesFullscreen=True
+bF11TogglesFullscreen=True
 bUseMouseForTouch=False
 bEnableMouseSmoothing=True
 bEnableFOVScaling=True
-FOVScale=0.011110
-DoubleClickTime=0.200000
 bCaptureMouseOnLaunch=True
-DefaultViewportMouseCaptureMode=CapturePermanently_IncludingInitialMouseDown
 bDefaultViewportMouseLock=True
-+ActionMappings=(ActionName="Jump",Key=SpaceBar,bShift=False,bCtrl=False,bAlt=False,bCmd=False)
-+ActionMappings=(ActionName="Jump",Key=Gamepad_FaceButton_Bottom,bShift=False,bCtrl=False,bAlt=False,bCmd=False)
-+ActionMappings=(ActionName="Jump",Key=MotionController_Left_Trigger,bShift=False,bCtrl=False,bAlt=False,bCmd=False)
-+ActionMappings=(ActionName="ResetVR",Key=R,bShift=False,bCtrl=False,bAlt=False,bCmd=False)
-+ActionMappings=(ActionName="ResetVR",Key=MotionController_Left_Grip1,bShift=False,bCtrl=False,bAlt=False,bCmd=False)
-+ActionMappings=(ActionName="Jump",Key=MotionController_Right_Trigger,bShift=False,bCtrl=False,bAlt=False,bCmd=False)
-+AxisMappings=(AxisName="MoveForward",Key=W,Scale=1.000000)
-+AxisMappings=(AxisName="MoveForward",Key=S,Scale=-1.000000)
-+AxisMappings=(AxisName="MoveForward",Key=Up,Scale=1.000000)
-+AxisMappings=(AxisName="MoveForward",Key=Down,Scale=-1.000000)
-+AxisMappings=(AxisName="MoveForward",Key=Gamepad_LeftY,Scale=1.000000)
-+AxisMappings=(AxisName="MoveRight",Key=A,Scale=-1.000000)
-+AxisMappings=(AxisName="MoveRight",Key=D,Scale=1.000000)
-+AxisMappings=(AxisName="MoveRight",Key=Gamepad_LeftX,Scale=1.000000)
-+AxisMappings=(AxisName="TurnRate",Key=Gamepad_RightX,Scale=1.000000)
-+AxisMappings=(AxisName="TurnRate",Key=Left,Scale=-1.000000)
-+AxisMappings=(AxisName="TurnRate",Key=Right,Scale=1.000000)
-+AxisMappings=(AxisName="Turn",Key=MouseX,Scale=1.000000)
-+AxisMappings=(AxisName="LookUpRate",Key=Gamepad_RightY,Scale=1.000000)
-+AxisMappings=(AxisName="LookUp",Key=MouseY,Scale=-1.000000)
-+AxisMappings=(AxisName="MoveForward",Key=MotionController_Left_Thumbstick_Y,Scale=-1.000000)
-+AxisMappings=(AxisName="MoveRight",Key=MotionController_Left_Thumbstick_X,Scale=1.000000)
-+AxisMappings=(AxisName="TurnRate",Key=MotionController_Right_Thumbstick_X,Scale=-1.000000)
 bAlwaysShowTouchInterface=False
 bShowConsoleOnFourFingerTap=True
+bEnableGestureRecognizer=False
+bUseAutocorrect=False
+DefaultViewportMouseCaptureMode=CapturePermanently_IncludingInitialMouseDown
+DefaultViewportMouseLockMode=LockOnCapture
+FOVScale=0.011110
+DoubleClickTime=0.200000
++ActionMappings=(ActionName="Jump",bShift=False,bCtrl=False,bAlt=False,bCmd=False,Key=SpaceBar)
++ActionMappings=(ActionName="Jump",bShift=False,bCtrl=False,bAlt=False,bCmd=False,Key=Gamepad_FaceButton_Bottom)
++ActionMappings=(ActionName="Jump",bShift=False,bCtrl=False,bAlt=False,bCmd=False,Key=MotionController_Left_Trigger)
++ActionMappings=(ActionName="ResetVR",bShift=False,bCtrl=False,bAlt=False,bCmd=False,Key=R)
++ActionMappings=(ActionName="ResetVR",bShift=False,bCtrl=False,bAlt=False,bCmd=False,Key=MotionController_Left_Grip1)
++ActionMappings=(ActionName="Jump",bShift=False,bCtrl=False,bAlt=False,bCmd=False,Key=MotionController_Right_Trigger)
++ActionMappings=(ActionName="Jump",bShift=False,bCtrl=False,bAlt=False,bCmd=False,Key=OculusTouchpad_Touchpad)
++AxisMappings=(AxisName="MoveForward",Scale=1.000000,Key=W)
++AxisMappings=(AxisName="MoveForward",Scale=-1.000000,Key=S)
++AxisMappings=(AxisName="MoveForward",Scale=1.000000,Key=Up)
++AxisMappings=(AxisName="MoveForward",Scale=-1.000000,Key=Down)
++AxisMappings=(AxisName="MoveForward",Scale=1.000000,Key=Gamepad_LeftY)
++AxisMappings=(AxisName="MoveRight",Scale=-1.000000,Key=A)
++AxisMappings=(AxisName="MoveRight",Scale=1.000000,Key=D)
++AxisMappings=(AxisName="MoveRight",Scale=1.000000,Key=Gamepad_LeftX)
++AxisMappings=(AxisName="TurnRate",Scale=1.000000,Key=Gamepad_RightX)
++AxisMappings=(AxisName="TurnRate",Scale=-1.000000,Key=Left)
++AxisMappings=(AxisName="TurnRate",Scale=1.000000,Key=Right)
++AxisMappings=(AxisName="Turn",Scale=1.000000,Key=MouseX)
++AxisMappings=(AxisName="LookUpRate",Scale=1.000000,Key=Gamepad_RightY)
++AxisMappings=(AxisName="LookUp",Scale=-1.000000,Key=MouseY)
++AxisMappings=(AxisName="MoveForward",Scale=-1.000000,Key=MotionController_Left_Thumbstick_Y)
++AxisMappings=(AxisName="MoveRight",Scale=1.000000,Key=MotionController_Left_Thumbstick_X)
++AxisMappings=(AxisName="TurnRate",Scale=-1.000000,Key=MotionController_Right_Thumbstick_X)
 DefaultTouchInterface=/Engine/MobileResources/HUD/DefaultVirtualJoysticks.DefaultVirtualJoysticks
 ConsoleKey=None
 -ConsoleKeys=Tilde
