First Person C++ Template overhaul.

diff --git a/Templates/TP_FirstPerson/Config/DefaultInput.ini b/Templates/TP_FirstPerson/Config/DefaultInput.ini
--- a/Templates/TP_FirstPerson/Config/DefaultInput.ini
+++ b/Templates/TP_FirstPerson/Config/DefaultInput.ini
@@ -18,6 +18,71 @@
 +AxisConfig=(AxisKeyName="Gamepad_RightTriggerAxis",AxisProperties=(DeadZone=0.000000,Sensitivity=1.000000,Exponent=1.000000,bInvert=False))
 +AxisConfig=(AxisKeyName="Gamepad_Special_Left_X",AxisProperties=(DeadZone=0.000000,Sensitivity=1.000000,Exponent=1.000000,bInvert=False))
 +AxisConfig=(AxisKeyName="Gamepad_Special_Left_Y",AxisProperties=(DeadZone=0.000000,Sensitivity=1.000000,Exponent=1.000000,bInvert=False))
++AxisConfig=(AxisKeyName="Vive_Left_Trigger_Axis",AxisProperties=(DeadZone=0.000000,Sensitivity=1.000000,Exponent=1.000000,bInvert=False))
++AxisConfig=(AxisKeyName="Vive_Left_Trackpad_X",AxisProperties=(DeadZone=0.000000,Sensitivity=1.000000,Exponent=1.000000,bInvert=False))
++AxisConfig=(AxisKeyName="Vive_Left_Trackpad_Y",AxisProperties=(DeadZone=0.000000,Sensitivity=1.000000,Exponent=1.000000,bInvert=False))
++AxisConfig=(AxisKeyName="Vive_Right_Trigger_Axis",AxisProperties=(DeadZone=0.000000,Sensitivity=1.000000,Exponent=1.000000,bInvert=False))
++AxisConfig=(AxisKeyName="Vive_Right_Trackpad_X",AxisProperties=(DeadZone=0.000000,Sensitivity=1.000000,Exponent=1.000000,bInvert=False))
++AxisConfig=(AxisKeyName="Vive_Right_Trackpad_Y",AxisProperties=(DeadZone=0.000000,Sensitivity=1.000000,Exponent=1.000000,bInvert=False))
++AxisConfig=(AxisKeyName="MixedReality_Left_Trigger_Axis",AxisProperties=(DeadZone=0.000000,Sensitivity=1.000000,Exponent=1.000000,bInvert=False))
++AxisConfig=(AxisKeyName="MixedReality_Left_Thumbstick_X",AxisProperties=(DeadZone=0.000000,Sensitivity=1.000000,Exponent=1.000000,bInvert=False))
++AxisConfig=(AxisKeyName="MixedReality_Left_Thumbstick_Y",AxisProperties=(DeadZone=0.000000,Sensitivity=1.000000,Exponent=1.000000,bInvert=False))
++AxisConfig=(AxisKeyName="MixedReality_Left_Trackpad_X",AxisProperties=(DeadZone=0.000000,Sensitivity=1.000000,Exponent=1.000000,bInvert=False))
++AxisConfig=(AxisKeyName="MixedReality_Left_Trackpad_Y",AxisProperties=(DeadZone=0.000000,Sensitivity=1.000000,Exponent=1.000000,bInvert=False))
++AxisConfig=(AxisKeyName="MixedReality_Right_Trigger_Axis",AxisProperties=(DeadZone=0.000000,Sensitivity=1.000000,Exponent=1.000000,bInvert=False))
++AxisConfig=(AxisKeyName="MixedReality_Right_Thumbstick_X",AxisProperties=(DeadZone=0.000000,Sensitivity=1.000000,Exponent=1.000000,bInvert=False))
++AxisConfig=(AxisKeyName="MixedReality_Right_Thumbstick_Y",AxisProperties=(DeadZone=0.000000,Sensitivity=1.000000,Exponent=1.000000,bInvert=False))
++AxisConfig=(AxisKeyName="MixedReality_Right_Trackpad_X",AxisProperties=(DeadZone=0.000000,Sensitivity=1.000000,Exponent=1.000000,bInvert=False))
++AxisConfig=(AxisKeyName="MixedReality_Right_Trackpad_Y",AxisProperties=(DeadZone=0.000000,Sensitivity=1.000000,Exponent=1.000000,bInvert=False))
++AxisConfig=(AxisKeyName="OculusTouch_Left_Grip_Axis",AxisProperties=(DeadZone=0.000000,Sensitivity=1.000000,Exponent=1.000000,bInvert=False))
++AxisConfig=(AxisKeyName="OculusTouch_Left_Trigger_Axis",AxisProperties=(DeadZone=0.000000,Sensitivity=1.000000,Exponent=1.000000,bInvert=False))
++AxisConfig=(AxisKeyName="OculusTouch_Left_Thumbstick_X",AxisProperties=(DeadZone=0.000000,Sensitivity=1.000000,Exponent=1.000000,bInvert=False))
++AxisConfig=(AxisKeyName="OculusTouch_Left_Thumbstick_Y",AxisProperties=(DeadZone=0.000000,Sensitivity=1.000000,Exponent=1.000000,bInvert=False))
++AxisConfig=(AxisKeyName="OculusTouch_Right_Grip_Axis",AxisProperties=(DeadZone=0.000000,Sensitivity=1.000000,Exponent=1.000000,bInvert=False))
++AxisConfig=(AxisKeyName="OculusTouch_Right_Trigger_Axis",AxisProperties=(DeadZone=0.000000,Sensitivity=1.000000,Exponent=1.000000,bInvert=False))
++AxisConfig=(AxisKeyName="OculusTouch_Right_Thumbstick_X",AxisProperties=(DeadZone=0.000000,Sensitivity=1.000000,Exponent=1.000000,bInvert=False))
++AxisConfig=(AxisKeyName="OculusTouch_Right_Thumbstick_Y",AxisProperties=(DeadZone=0.000000,Sensitivity=1.000000,Exponent=1.000000,bInvert=False))
++AxisConfig=(AxisKeyName="ValveIndex_Left_Grip_Axis",AxisProperties=(DeadZone=0.000000,Sensitivity=1.000000,Exponent=1.000000,bInvert=False))
++AxisConfig=(AxisKeyName="ValveIndex_Left_Grip_Force",AxisProperties=(DeadZone=0.000000,Sensitivity=1.000000,Exponent=1.000000,bInvert=False))
++AxisConfig=(AxisKeyName="ValveIndex_Left_Trigger_Axis",AxisProperties=(DeadZone=0.000000,Sensitivity=1.000000,Exponent=1.000000,bInvert=False))
++AxisConfig=(AxisKeyName="ValveIndex_Left_Thumbstick_X",AxisProperties=(DeadZone=0.000000,Sensitivity=1.000000,Exponent=1.000000,bInvert=False))
++AxisConfig=(AxisKeyName="ValveIndex_Left_Thumbstick_Y",AxisProperties=(DeadZone=0.000000,Sensitivity=1.000000,Exponent=1.000000,bInvert=False))
++AxisConfig=(AxisKeyName="ValveIndex_Left_Trackpad_X",AxisProperties=(DeadZone=0.000000,Sensitivity=1.000000,Exponent=1.000000,bInvert=False))
++AxisConfig=(AxisKeyName="ValveIndex_Left_Trackpad_Y",AxisProperties=(DeadZone=0.000000,Sensitivity=1.000000,Exponent=1.000000,bInvert=False))
++AxisConfig=(AxisKeyName="ValveIndex_Left_Trackpad_Force",AxisProperties=(DeadZone=0.000000,Sensitivity=1.000000,Exponent=1.000000,bInvert=False))
++AxisConfig=(AxisKeyName="ValveIndex_Left_Trackpad_Touch",AxisProperties=(DeadZone=0.000000,Sensitivity=1.000000,Exponent=1.000000,bInvert=False))
++AxisConfig=(AxisKeyName="ValveIndex_Right_Grip_Axis",AxisProperties=(DeadZone=0.000000,Sensitivity=1.000000,Exponent=1.000000,bInvert=False))
++AxisConfig=(AxisKeyName="ValveIndex_Right_Grip_Force",AxisProperties=(DeadZone=0.000000,Sensitivity=1.000000,Exponent=1.000000,bInvert=False))
++AxisConfig=(AxisKeyName="ValveIndex_Right_Trigger_Axis",AxisProperties=(DeadZone=0.000000,Sensitivity=1.000000,Exponent=1.000000,bInvert=False))
++AxisConfig=(AxisKeyName="ValveIndex_Right_Thumbstick_X",AxisProperties=(DeadZone=0.000000,Sensitivity=1.000000,Exponent=1.000000,bInvert=False))
++AxisConfig=(AxisKeyName="ValveIndex_Right_Thumbstick_Y",AxisProperties=(DeadZone=0.000000,Sensitivity=1.000000,Exponent=1.000000,bInvert=False))
++AxisConfig=(AxisKeyName="ValveIndex_Right_Trackpad_X",AxisProperties=(DeadZone=0.000000,Sensitivity=1.000000,Exponent=1.000000,bInvert=False))
++AxisConfig=(AxisKeyName="ValveIndex_Right_Trackpad_Y",AxisProperties=(DeadZone=0.000000,Sensitivity=1.000000,Exponent=1.000000,bInvert=False))
++AxisConfig=(AxisKeyName="ValveIndex_Right_Trackpad_Force",AxisProperties=(DeadZone=0.000000,Sensitivity=1.000000,Exponent=1.000000,bInvert=False))
++AxisConfig=(AxisKeyName="OculusTouch_Left_Thumbstick",AxisProperties=(DeadZone=0.000000,Sensitivity=1.000000,Exponent=1.000000,bInvert=False))
++AxisConfig=(AxisKeyName="OculusTouch_Left_FaceButton1",AxisProperties=(DeadZone=0.000000,Sensitivity=1.000000,Exponent=1.000000,bInvert=False))
++AxisConfig=(AxisKeyName="OculusTouch_Left_Trigger",AxisProperties=(DeadZone=0.000000,Sensitivity=1.000000,Exponent=1.000000,bInvert=False))
++AxisConfig=(AxisKeyName="OculusTouch_Left_FaceButton2",AxisProperties=(DeadZone=0.000000,Sensitivity=1.000000,Exponent=1.000000,bInvert=False))
++AxisConfig=(AxisKeyName="OculusTouch_Left_IndexPointing",AxisProperties=(DeadZone=0.000000,Sensitivity=1.000000,Exponent=1.000000,bInvert=False))
++AxisConfig=(AxisKeyName="OculusTouch_Left_ThumbUp",AxisProperties=(DeadZone=0.000000,Sensitivity=1.000000,Exponent=1.000000,bInvert=False))
++AxisConfig=(AxisKeyName="OculusTouch_Left_ThumbRest",AxisProperties=(DeadZone=0.000000,Sensitivity=1.000000,Exponent=1.000000,bInvert=False))
++AxisConfig=(AxisKeyName="OculusTouch_Right_Thumbstick",AxisProperties=(DeadZone=0.000000,Sensitivity=1.000000,Exponent=1.000000,bInvert=False))
++AxisConfig=(AxisKeyName="OculusTouch_Right_FaceButton1",AxisProperties=(DeadZone=0.000000,Sensitivity=1.000000,Exponent=1.000000,bInvert=False))
++AxisConfig=(AxisKeyName="OculusTouch_Right_Trigger",AxisProperties=(DeadZone=0.000000,Sensitivity=1.000000,Exponent=1.000000,bInvert=False))
++AxisConfig=(AxisKeyName="OculusTouch_Right_FaceButton2",AxisProperties=(DeadZone=0.000000,Sensitivity=1.000000,Exponent=1.000000,bInvert=False))
++AxisConfig=(AxisKeyName="OculusTouch_Right_IndexPointing",AxisProperties=(DeadZone=0.000000,Sensitivity=1.000000,Exponent=1.000000,bInvert=False))
++AxisConfig=(AxisKeyName="OculusTouch_Right_ThumbUp",AxisProperties=(DeadZone=0.000000,Sensitivity=1.000000,Exponent=1.000000,bInvert=False))
++AxisConfig=(AxisKeyName="OculusTouch_Right_ThumbRest",AxisProperties=(DeadZone=0.000000,Sensitivity=1.000000,Exponent=1.000000,bInvert=False))
++AxisConfig=(AxisKeyName="OculusHand_Left_ThumbPinchStrength",AxisProperties=(DeadZone=0.000000,Sensitivity=1.000000,Exponent=1.000000,bInvert=False))
++AxisConfig=(AxisKeyName="OculusHand_Left_IndexPinchStrength",AxisProperties=(DeadZone=0.000000,Sensitivity=1.000000,Exponent=1.000000,bInvert=False))
++AxisConfig=(AxisKeyName="OculusHand_Left_MiddlePinchStrength",AxisProperties=(DeadZone=0.000000,Sensitivity=1.000000,Exponent=1.000000,bInvert=False))
++AxisConfig=(AxisKeyName="OculusHand_Left_RingPinchStrength",AxisProperties=(DeadZone=0.000000,Sensitivity=1.000000,Exponent=1.000000,bInvert=False))
++AxisConfig=(AxisKeyName="OculusHand_Left_PinkPinchStrength",AxisProperties=(DeadZone=0.000000,Sensitivity=1.000000,Exponent=1.000000,bInvert=False))
++AxisConfig=(AxisKeyName="OculusHand_Right_ThumbPinchStrength",AxisProperties=(DeadZone=0.000000,Sensitivity=1.000000,Exponent=1.000000,bInvert=False))
++AxisConfig=(AxisKeyName="OculusHand_Right_IndexPinchStrength",AxisProperties=(DeadZone=0.000000,Sensitivity=1.000000,Exponent=1.000000,bInvert=False))
++AxisConfig=(AxisKeyName="OculusHand_Right_MiddlePinchStrength",AxisProperties=(DeadZone=0.000000,Sensitivity=1.000000,Exponent=1.000000,bInvert=False))
++AxisConfig=(AxisKeyName="OculusHand_Right_RingPinchStrength",AxisProperties=(DeadZone=0.000000,Sensitivity=1.000000,Exponent=1.000000,bInvert=False))
++AxisConfig=(AxisKeyName="OculusHand_Right_PinkPinchStrength",AxisProperties=(DeadZone=0.000000,Sensitivity=1.000000,Exponent=1.000000,bInvert=False))
 bAltEnterTogglesFullscreen=True
 bF11TogglesFullscreen=True
 bUseMouseForTouch=False
@@ -35,22 +100,22 @@ FOVScale=0.011110
 DoubleClickTime=0.200000
 +ActionMappings=(ActionName="Jump",bShift=False,bCtrl=False,bAlt=False,bCmd=False,Key=SpaceBar)
 +ActionMappings=(ActionName="Jump",bShift=False,bCtrl=False,bAlt=False,bCmd=False,Key=Gamepad_FaceButton_Bottom)
-+ActionMappings=(ActionName="Fire",bShift=False,bCtrl=False,bAlt=False,bCmd=False,Key=LeftMouseButton)
-+ActionMappings=(ActionName="Fire",bShift=False,bCtrl=False,bAlt=False,bCmd=False,Key=Gamepad_RightTrigger)
-+AxisMappings=(AxisName="MoveForward",Scale=1.000000,Key=W)
-+AxisMappings=(AxisName="MoveForward",Scale=-1.000000,Key=S)
-+AxisMappings=(AxisName="MoveForward",Scale=1.000000,Key=Up)
-+AxisMappings=(AxisName="MoveForward",Scale=-1.000000,Key=Down)
-+AxisMappings=(AxisName="MoveForward",Scale=1.000000,Key=Gamepad_LeftY)
-+AxisMappings=(AxisName="MoveRight",Scale=-1.000000,Key=A)
-+AxisMappings=(AxisName="MoveRight",Scale=1.000000,Key=D)
-+AxisMappings=(AxisName="MoveRight",Scale=1.000000,Key=Gamepad_LeftX)
-+AxisMappings=(AxisName="TurnRate",Scale=1.000000,Key=Gamepad_RightX)
-+AxisMappings=(AxisName="TurnRate",Scale=-1.000000,Key=Left)
-+AxisMappings=(AxisName="TurnRate",Scale=1.000000,Key=Right)
-+AxisMappings=(AxisName="Turn",Scale=1.000000,Key=MouseX)
-+AxisMappings=(AxisName="LookUpRate",Scale=1.000000,Key=Gamepad_RightY)
-+AxisMappings=(AxisName="LookUp",Scale=-1.000000,Key=MouseY)
++ActionMappings=(ActionName="PrimaryAction",bShift=False,bCtrl=False,bAlt=False,bCmd=False,Key=LeftMouseButton)
++ActionMappings=(ActionName="PrimaryAction",bShift=False,bCtrl=False,bAlt=False,bCmd=False,Key=Gamepad_RightTrigger)
++AxisMappings=(AxisName="Move Forward / Backward",Scale=1.000000,Key=W)
++AxisMappings=(AxisName="Move Forward / Backward",Scale=-1.000000,Key=S)
++AxisMappings=(AxisName="Move Forward / Backward",Scale=1.000000,Key=Up)
++AxisMappings=(AxisName="Move Forward / Backward",Scale=-1.000000,Key=Down)
++AxisMappings=(AxisName="Move Forward / Backward",Scale=1.000000,Key=Gamepad_LeftY)
++AxisMappings=(AxisName="Move Right / Left",Scale=-1.000000,Key=A)
++AxisMappings=(AxisName="Move Right / Left",Scale=1.000000,Key=D)
++AxisMappings=(AxisName="Move Right / Left",Scale=1.000000,Key=Gamepad_LeftX)
++AxisMappings=(AxisName="Turn Right / Left Gamepad",Scale=1.000000,Key=Gamepad_RightX)
++AxisMappings=(AxisName="Turn Right / Left Gamepad",Scale=-1.000000,Key=Left)
++AxisMappings=(AxisName="Turn Right / Left Gamepad",Scale=1.000000,Key=Right)
++AxisMappings=(AxisName="Turn Right / Left Mouse",Scale=1.000000,Key=MouseX)
++AxisMappings=(AxisName="Look Up / Down Gamepad",Scale=1.000000,Key=Gamepad_RightY)
++AxisMappings=(AxisName="Look Up / Down Mouse",Scale=-1.000000,Key=MouseY)
 DefaultPlayerInputClass=/Script/Engine.PlayerInput
 DefaultInputComponentClass=/Script/Engine.InputComponent
 DefaultTouchInterface=/Engine/MobileResources/HUD/DefaultVirtualJoysticks.DefaultVirtualJoysticks
diff --git a/Templates/TP_FirstPerson/Config/config.ini b/Templates/TP_FirstPerson/Config/config.ini
--- a/Templates/TP_FirstPerson/Config/config.ini
+++ b/Templates/TP_FirstPerson/Config/config.ini
@@ -1,22 +1,22 @@
 [InputSettings]
-+AxisMappings=(AxisName="MoveForward",Scale=1.000000,Key=W)
-+AxisMappings=(AxisName="MoveForward",Scale=-1.000000,Key=S)
-+AxisMappings=(AxisName="MoveForward",Scale=1.000000,Key=Up)
-+AxisMappings=(AxisName="MoveForward",Scale=-1.000000,Key=Down)
-+AxisMappings=(AxisName="MoveForward",Scale=1.000000,Key=Gamepad_LeftY)
-+AxisMappings=(AxisName="MoveRight",Scale=-1.000000,Key=A)
-+AxisMappings=(AxisName="MoveRight",Scale=1.000000,Key=D)
-+AxisMappings=(AxisName="MoveRight",Scale=1.000000,Key=Gamepad_LeftX)
-+AxisMappings=(AxisName="TurnRate",Scale=1.000000,Key=Gamepad_RightX)
-+AxisMappings=(AxisName="TurnRate",Scale=-1.000000,Key=Left)
-+AxisMappings=(AxisName="TurnRate",Scale=1.000000,Key=Right)
-+AxisMappings=(AxisName="Turn",Scale=1.000000,Key=MouseX)
-+AxisMappings=(AxisName="LookUpRate",Scale=1.000000,Key=Gamepad_RightY)
-+AxisMappings=(AxisName="LookUp",Scale=-1.000000,Key=MouseY)
 +ActionMappings=(ActionName="Jump",bShift=False,bCtrl=False,bAlt=False,bCmd=False,Key=SpaceBar)
 +ActionMappings=(ActionName="Jump",bShift=False,bCtrl=False,bAlt=False,bCmd=False,Key=Gamepad_FaceButton_Bottom)
-+ActionMappings=(ActionName="Fire",bShift=False,bCtrl=False,bAlt=False,bCmd=False,Key=LeftMouseButton)
-+ActionMappings=(ActionName="Fire",bShift=False,bCtrl=False,bAlt=False,bCmd=False,Key=Gamepad_RightTrigger)
++ActionMappings=(ActionName="PrimaryAction",bShift=False,bCtrl=False,bAlt=False,bCmd=False,Key=LeftMouseButton)
++ActionMappings=(ActionName="PrimaryAction",bShift=False,bCtrl=False,bAlt=False,bCmd=False,Key=Gamepad_RightTrigger)
++AxisMappings=(AxisName="Move Forward / Backward",Scale=1.000000,Key=W)
++AxisMappings=(AxisName="Move Forward / Backward",Scale=-1.000000,Key=S)
++AxisMappings=(AxisName="Move Forward / Backward",Scale=1.000000,Key=Up)
++AxisMappings=(AxisName="Move Forward / Backward",Scale=-1.000000,Key=Down)
++AxisMappings=(AxisName="Move Forward / Backward",Scale=1.000000,Key=Gamepad_LeftY)
++AxisMappings=(AxisName="Move Right / Left",Scale=-1.000000,Key=A)
++AxisMappings=(AxisName="Move Right / Left",Scale=1.000000,Key=D)
++AxisMappings=(AxisName="Move Right / Left",Scale=1.000000,Key=Gamepad_LeftX)
++AxisMappings=(AxisName="Turn Right / Left Gamepad",Scale=1.000000,Key=Gamepad_RightX)
++AxisMappings=(AxisName="Turn Right / Left Gamepad",Scale=-1.000000,Key=Left)
++AxisMappings=(AxisName="Turn Right / Left Gamepad",Scale=1.000000,Key=Right)
++AxisMappings=(AxisName="Turn Right / Left Mouse",Scale=1.000000,Key=MouseX)
++AxisMappings=(AxisName="Look Up / Down Gamepad",Scale=1.000000,Key=Gamepad_RightY)
++AxisMappings=(AxisName="Look Up / Down Mouse",Scale=-1.000000,Key=MouseY)
 
 [Redirects]
 GameName=TP_FirstPerson
diff --git a/Templates/TP_FirstPerson/Source/TP_FirstPerson/TP_FirstPersonCharacter.cpp b/Templates/TP_FirstPerson/Source/TP_FirstPerson/TP_FirstPersonCharacter.cpp
--- a/Templates/TP_FirstPerson/Source/TP_FirstPerson/TP_FirstPersonCharacter.cpp
+++ b/Templates/TP_FirstPerson/Source/TP_FirstPerson/TP_FirstPersonCharacter.cpp
@@ -7,9 +7,7 @@
 #include "Components/CapsuleComponent.h"
 #include "Components/InputComponent.h"
 #include "GameFramework/InputSettings.h"
-#include "Kismet/GameplayStatics.h"
 
-DEFINE_LOG_CATEGORY_STATIC(LogFPChar, Warning, All);
 
 //////////////////////////////////////////////////////////////////////////
 // ATP_FirstPersonCharacter
@@ -20,8 +18,7 @@ ATP_FirstPersonCharacter::ATP_FirstPersonCharacter()
 	GetCapsuleComponent()->InitCapsuleSize(55.f, 96.0f);
 
 	// set our turn rates for input
-	BaseTurnRate = 45.f;
-	BaseLookUpRate = 45.f;
+	TurnRateGamepad = 45.f;
 
 	// Create a CameraComponent	
 	FirstPersonCameraComponent = CreateDefaultSubobject<UCameraComponent>(TEXT("FirstPersonCamera"));
@@ -38,23 +35,6 @@ ATP_FirstPersonCharacter::ATP_FirstPersonCharacter()
 	Mesh1P->SetRelativeRotation(FRotator(1.9f, -19.19f, 5.2f));
 	Mesh1P->SetRelativeLocation(FVector(-0.5f, -4.4f, -155.7f));
 
-	// Create a gun mesh component
-	FP_Gun = CreateDefaultSubobject<USkeletalMeshComponent>(TEXT("FP_Gun"));
-	FP_Gun->SetOnlyOwnerSee(false);			// otherwise won't be visible in multiplayer
-	FP_Gun->bCastDynamicShadow = false;
-	FP_Gun->CastShadow = false;
-	// FP_Gun->SetupAttachment(Mesh1P, TEXT("GripPoint"));
-	FP_Gun->SetupAttachment(RootComponent);
-
-	FP_MuzzleLocation = CreateDefaultSubobject<USceneComponent>(TEXT("MuzzleLocation"));
-	FP_MuzzleLocation->SetupAttachment(FP_Gun);
-	FP_MuzzleLocation->SetRelativeLocation(FVector(0.2f, 48.4f, -10.6f));
-
-	// Default offset from the character location for projectiles to spawn
-	GunOffset = FVector(100.0f, 0.0f, 10.0f);
-
-	// Note: The ProjectileClass and the skeletal mesh/anim blueprints for Mesh1P and FP_Gun
-	// are set in the derived blueprint asset named FirstPersonCharacter to avoid direct content references in C++.
 }
 
 void ATP_FirstPersonCharacter::BeginPlay()
@@ -62,8 +42,6 @@ void ATP_FirstPersonCharacter::BeginPlay()
 	// Call the base class  
 	Super::BeginPlay();
 
-	//Attach gun mesh component to Skeleton, doing it here because the skeleton is not yet created in the constructor
-	FP_Gun->AttachToComponent(Mesh1P, FAttachmentTransformRules(EAttachmentRule::SnapToTarget, true), TEXT("GripPoint"));
 }
 
 //////////////////////////////////////////////////////////////////////////// Input
@@ -78,61 +56,28 @@ void ATP_FirstPersonCharacter::SetupPlayerInputComponent(class UInputComponent*
 	PlayerInputComponent->BindAction("Jump", IE_Released, this, &ACharacter::StopJumping);
 
 	// Bind fire event
-	PlayerInputComponent->BindAction("Fire", IE_Pressed, this, &ATP_FirstPersonCharacter::OnFire);
+	PlayerInputComponent->BindAction("PrimaryAction", IE_Pressed, this, &ATP_FirstPersonCharacter::OnPrimaryAction);
 
 	// Enable touchscreen input
 	EnableTouchscreenMovement(PlayerInputComponent);
 
 	// Bind movement events
-	PlayerInputComponent->BindAxis("MoveForward", this, &ATP_FirstPersonCharacter::MoveForward);
-	PlayerInputComponent->BindAxis("MoveRight", this, &ATP_FirstPersonCharacter::MoveRight);
+	PlayerInputComponent->BindAxis("Move Forward / Backward", this, &ATP_FirstPersonCharacter::MoveForward);
+	PlayerInputComponent->BindAxis("Move Right / Left", this, &ATP_FirstPersonCharacter::MoveRight);
 
 	// We have 2 versions of the rotation bindings to handle different kinds of devices differently
-	// "turn" handles devices that provide an absolute delta, such as a mouse.
-	// "turnrate" is for devices that we choose to treat as a rate of change, such as an analog joystick
-	PlayerInputComponent->BindAxis("Turn", this, &APawn::AddControllerYawInput);
-	PlayerInputComponent->BindAxis("TurnRate", this, &ATP_FirstPersonCharacter::TurnAtRate);
-	PlayerInputComponent->BindAxis("LookUp", this, &APawn::AddControllerPitchInput);
-	PlayerInputComponent->BindAxis("LookUpRate", this, &ATP_FirstPersonCharacter::LookUpAtRate);
+	// "Mouse" versions handle devices that provide an absolute delta, such as a mouse.
+	// "Gamepad" versions are for devices that we choose to treat as a rate of change, such as an analog joystick
+	PlayerInputComponent->BindAxis("Turn Right / Left Mouse", this, &APawn::AddControllerYawInput);
+	PlayerInputComponent->BindAxis("Look Up / Down Mouse", this, &APawn::AddControllerPitchInput);
+	PlayerInputComponent->BindAxis("Turn Right / Left Gamepad", this, &ATP_FirstPersonCharacter::TurnAtRate);
+	PlayerInputComponent->BindAxis("Look Up / Down Gamepad", this, &ATP_FirstPersonCharacter::LookUpAtRate);
 }
 
-void ATP_FirstPersonCharacter::OnFire()
+void ATP_FirstPersonCharacter::OnPrimaryAction()
 {
-	// Try and fire a projectile
-	if (ProjectileClass != nullptr)
-	{
-		UWorld* const World = GetWorld();
-		if (World != nullptr)
-		{
-			const FRotator SpawnRotation = GetControlRotation();
-			// MuzzleOffset is in camera space, so transform it to world space before offsetting from the character location to find the final muzzle position
-			const FVector SpawnLocation = ((FP_MuzzleLocation != nullptr) ? FP_MuzzleLocation->GetComponentLocation() : GetActorLocation()) + SpawnRotation.RotateVector(GunOffset);
-
-			//Set Spawn Collision Handling Override
-			FActorSpawnParameters ActorSpawnParams;
-			ActorSpawnParams.SpawnCollisionHandlingOverride = ESpawnActorCollisionHandlingMethod::AdjustIfPossibleButDontSpawnIfColliding;
-
-			// Spawn the projectile at the muzzle
-			World->SpawnActor<ATP_FirstPersonProjectile>(ProjectileClass, SpawnLocation, SpawnRotation, ActorSpawnParams);
-		}
-	}
-
-	// Try and play the sound if specified
-	if (FireSound != nullptr)
-	{
-		UGameplayStatics::PlaySoundAtLocation(this, FireSound, GetActorLocation());
-	}
-
-	// Try and play a firing animation if specified
-	if (FireAnimation != nullptr)
-	{
-		// Get the animation object for the arms mesh
-		UAnimInstance* AnimInstance = Mesh1P->GetAnimInstance();
-		if (AnimInstance != nullptr)
-		{
-			AnimInstance->Montage_Play(FireAnimation, 1.f);
-		}
-	}
+	// Trigger the OnItemUsed Event
+	OnItemUsed.Broadcast();
 }
 
 void ATP_FirstPersonCharacter::BeginTouch(const ETouchIndex::Type FingerIndex, const FVector Location)
@@ -143,7 +88,7 @@ void ATP_FirstPersonCharacter::BeginTouch(const ETouchIndex::Type FingerIndex, c
 	}
 	if ((FingerIndex == TouchItem.FingerIndex) && (TouchItem.bMoved == false))
 	{
-		OnFire();
+		OnPrimaryAction();
 	}
 	TouchItem.bIsPressed = true;
 	TouchItem.FingerIndex = FingerIndex;
@@ -181,13 +126,13 @@ void ATP_FirstPersonCharacter::EndTouch(const ETouchIndex::Type FingerIndex, con
 //					if (FMath::Abs(ScaledDelta.X) >= 4.0 / ScreenSize.X)
 //					{
 //						TouchItem.bMoved = true;
-//						float Value = ScaledDelta.X * BaseTurnRate;
+//						float Value = ScaledDelta.X * TurnRateGamepad;
 //						AddControllerYawInput(Value);
 //					}
 //					if (FMath::Abs(ScaledDelta.Y) >= 4.0 / ScreenSize.Y)
 //					{
 //						TouchItem.bMoved = true;
-//						float Value = ScaledDelta.Y * BaseTurnRate;
+//						float Value = ScaledDelta.Y * TurnRateGamepad;
 //						AddControllerPitchInput(Value);
 //					}
 //					TouchItem.Location = Location;
@@ -219,13 +164,13 @@ void ATP_FirstPersonCharacter::MoveRight(float Value)
 void ATP_FirstPersonCharacter::TurnAtRate(float Rate)
 {
 	// calculate delta for this frame from the rate information
-	AddControllerYawInput(Rate * BaseTurnRate * GetWorld()->GetDeltaSeconds());
+	AddControllerYawInput(Rate * TurnRateGamepad * GetWorld()->GetDeltaSeconds());
 }
 
 void ATP_FirstPersonCharacter::LookUpAtRate(float Rate)
 {
 	// calculate delta for this frame from the rate information
-	AddControllerPitchInput(Rate * BaseLookUpRate * GetWorld()->GetDeltaSeconds());
+	AddControllerPitchInput(Rate * TurnRateGamepad * GetWorld()->GetDeltaSeconds());
 }
 
 bool ATP_FirstPersonCharacter::EnableTouchscreenMovement(class UInputComponent* PlayerInputComponent)
diff --git a/Templates/TP_FirstPerson/Source/TP_FirstPerson/TP_FirstPersonCharacter.h b/Templates/TP_FirstPerson/Source/TP_FirstPerson/TP_FirstPersonCharacter.h
--- a/Templates/TP_FirstPerson/Source/TP_FirstPerson/TP_FirstPersonCharacter.h
+++ b/Templates/TP_FirstPerson/Source/TP_FirstPerson/TP_FirstPersonCharacter.h
@@ -13,6 +13,8 @@ class UCameraComponent;
 class UAnimMontage;
 class USoundBase;
 
+DECLARE_DYNAMIC_MULTICAST_DELEGATE(FOnItemUsed);
+
 UCLASS(config=Game)
 class ATP_FirstPersonCharacter : public ACharacter
 {
@@ -22,14 +24,6 @@ class ATP_FirstPersonCharacter : public ACharacter
 	UPROPERTY(VisibleDefaultsOnly, Category=Mesh)
 	USkeletalMeshComponent* Mesh1P;
 
-	/** Gun mesh: 1st person view (seen only by self) */
-	UPROPERTY(VisibleDefaultsOnly, Category = Mesh)
-	USkeletalMeshComponent* FP_Gun;
-
-	/** Location on gun mesh where projectiles should spawn. */
-	UPROPERTY(VisibleDefaultsOnly, Category = Mesh)
-	USceneComponent* FP_MuzzleLocation;
-
 	/** First person camera */
 	UPROPERTY(VisibleAnywhere, BlueprintReadOnly, Category = Camera, meta = (AllowPrivateAccess = "true"))
 	UCameraComponent* FirstPersonCameraComponent;
@@ -43,32 +37,14 @@ protected:
 public:
 	/** Base turn rate, in deg/sec. Other scaling may affect final turn rate. */
 	UPROPERTY(VisibleAnywhere, BlueprintReadOnly, Category=Camera)
-	float BaseTurnRate;
-
-	/** Base look up/down rate, in deg/sec. Other scaling may affect final rate. */
-	UPROPERTY(VisibleAnywhere, BlueprintReadOnly, Category=Camera)
-	float BaseLookUpRate;
-
-	/** Gun muzzle's offset from the characters location */
-	UPROPERTY(EditAnywhere, BlueprintReadWrite, Category=Gameplay)
-	FVector GunOffset;
-
-	/** Projectile class to spawn */
-	UPROPERTY(EditDefaultsOnly, Category=Projectile)
-	TSubclassOf<class ATP_FirstPersonProjectile> ProjectileClass;
-
-	/** Sound to play each time we fire */
-	UPROPERTY(EditAnywhere, BlueprintReadWrite, Category=Gameplay)
-	USoundBase* FireSound;
-
-	/** AnimMontage to play each time we fire */
-	UPROPERTY(EditAnywhere, BlueprintReadWrite, Category = Gameplay)
-	UAnimMontage* FireAnimation;
+	float TurnRateGamepad;
 
+	UPROPERTY(BlueprintAssignable, Category = "Interaction")
+	FOnItemUsed OnItemUsed;
 protected:
 	
 	/** Fires a projectile. */
-	void OnFire();
+	void OnPrimaryAction();
 
 	/** Handles moving forward/backward */
 	void MoveForward(float Val);
diff --git a/Templates/TP_FirstPerson/Source/TP_FirstPerson/TP_FirstPersonGameMode.cpp b/Templates/TP_FirstPerson/Source/TP_FirstPerson/TP_FirstPersonGameMode.cpp
--- a/Templates/TP_FirstPerson/Source/TP_FirstPerson/TP_FirstPersonGameMode.cpp
+++ b/Templates/TP_FirstPerson/Source/TP_FirstPerson/TP_FirstPersonGameMode.cpp
@@ -8,7 +8,7 @@ ATP_FirstPersonGameMode::ATP_FirstPersonGameMode()
 	: Super()
 {
 	// set default pawn class to our Blueprinted character
-	static ConstructorHelpers::FClassFinder<APawn> PlayerPawnClassFinder(TEXT("/Game/FirstPerson/Blueprints/FirstPersonCharacter"));
+	static ConstructorHelpers::FClassFinder<APawn> PlayerPawnClassFinder(TEXT("/Game/FirstPerson/Blueprints/BP_FirstPersonCharacter"));
 	DefaultPawnClass = PlayerPawnClassFinder.Class;
 
 }
diff --git a/Templates/TP_FirstPerson/Source/TP_FirstPerson/TP_WeaponComponent.cpp b/Templates/TP_FirstPerson/Source/TP_FirstPerson/TP_WeaponComponent.cpp
new file mode 100644
--- /dev/null
+++ b/Templates/TP_FirstPerson/Source/TP_FirstPerson/TP_WeaponComponent.cpp
@@ -0,0 +1,90 @@
+// Copyright Epic Games, Inc. All Rights Reserved.
+
+
+#include "TP_WeaponComponent.h"
+#include "TP_FirstPersonCharacter.h"
+#include "TP_FirstPersonProjectile.h"
+#include "GameFramework/PlayerController.h"
+#include "Camera/PlayerCameraManager.h"
+#include "Kismet/GameplayStatics.h"
+
+// Sets default values for this component's properties
+UTP_WeaponComponent::UTP_WeaponComponent()
+{
+	SphereRadius = 32.f;
+	// Default offset from the character location for projectiles to spawn
+	MuzzleOffset = FVector(100.0f, 0.0f, 10.0f);
+}
+
+
+void UTP_WeaponComponent::Fire()
+{
+	if(Character == nullptr || Character->GetController() == nullptr)
+	{
+		return;
+	}
+
+	// Try and fire a projectile
+	if (ProjectileClass != nullptr)
+	{
+		UWorld* const World = GetWorld();
+		if (World != nullptr)
+		{
+			APlayerController* PlayerController = Cast<APlayerController>(Character->GetController());
+			const FRotator SpawnRotation = PlayerController->PlayerCameraManager->GetCameraRotation();
+			// MuzzleOffset is in camera space, so transform it to world space before offsetting from the character location to find the final muzzle position
+			const FVector SpawnLocation = GetOwner()->GetActorLocation() + SpawnRotation.RotateVector(MuzzleOffset);
+	
+			//Set Spawn Collision Handling Override
+			FActorSpawnParameters ActorSpawnParams;
+			ActorSpawnParams.SpawnCollisionHandlingOverride = ESpawnActorCollisionHandlingMethod::AdjustIfPossibleButDontSpawnIfColliding;
+	
+			// Spawn the projectile at the muzzle
+			World->SpawnActor<ATP_FirstPersonProjectile>(ProjectileClass, SpawnLocation, SpawnRotation, ActorSpawnParams);
+		}
+	}
+	
+	// Try and play the sound if specified
+	if (FireSound != nullptr)
+	{
+		UGameplayStatics::PlaySoundAtLocation(this, FireSound, Character->GetActorLocation());
+	}
+	
+	// Try and play a firing animation if specified
+	if (FireAnimation != nullptr)
+	{
+		// Get the animation object for the arms mesh
+		UAnimInstance* AnimInstance = Character->GetMesh1P()->GetAnimInstance();
+		if (AnimInstance != nullptr)
+		{
+			AnimInstance->Montage_Play(FireAnimation, 1.f);
+		}
+	}
+}
+
+// Called when the game starts
+void UTP_WeaponComponent::BeginPlay()
+{
+	Super::BeginPlay();
+
+	// Register our Overlap Event
+	OnComponentBeginOverlap.AddDynamic(this, &UTP_WeaponComponent::OnSphereBeginOverlap);
+}
+
+void UTP_WeaponComponent::OnSphereBeginOverlap(UPrimitiveComponent* OverlappedComponent, AActor* OtherActor, UPrimitiveComponent* OtherComp, int32 OtherBodyIndex, bool bFromSweep, const FHitResult& SweepResult)
+{
+	// Checking if it is a First Person Character overlapping
+	Character = Cast<ATP_FirstPersonCharacter>(OtherActor);
+	if(Character != nullptr)
+	{
+		// Attach the weapon to the First Person Character
+		FAttachmentTransformRules AttachmentRules(EAttachmentRule::SnapToTarget, true);
+		GetOwner()->AttachToComponent(Character->GetMesh1P(),AttachmentRules, FName(TEXT("GripPoint")));
+
+		// Unregister from the Overlap Event so it is no longer triggered
+		OnComponentBeginOverlap.RemoveAll(this);
+		// Register so that Fire is called every time the character tries to use the item being held
+		Character->OnItemUsed.AddDynamic(this, &UTP_WeaponComponent::Fire);
+	}
+}
+
diff --git a/Templates/TP_FirstPerson/Source/TP_FirstPerson/TP_WeaponComponent.h b/Templates/TP_FirstPerson/Source/TP_FirstPerson/TP_WeaponComponent.h
new file mode 100644
--- /dev/null
+++ b/Templates/TP_FirstPerson/Source/TP_FirstPerson/TP_WeaponComponent.h
@@ -0,0 +1,57 @@
+// Copyright Epic Games, Inc. All Rights Reserved.
+
+#pragma once
+
+#include "CoreMinimal.h"
+#include "Components/ActorComponent.h"
+#include "Components/SphereComponent.h"
+#include "TP_WeaponComponent.generated.h"
+
+class ATP_FirstPersonCharacter;
+class USkeletalMeshComponent;
+class USceneComponent;
+
+UCLASS(Blueprintable, BlueprintType, ClassGroup=(Custom), meta=(BlueprintSpawnableComponent) )
+class TP_FIRSTPERSON_API UTP_WeaponComponent : public USphereComponent
+{
+	GENERATED_BODY()
+	/** Sphere collision to do Pick up */
+	UPROPERTY(VisibleAnywhere, BlueprintReadOnly, Category = Mesh, meta = (AllowPrivateAccess = "true"))
+	USphereComponent* SphereCollision;
+
+public:
+	/** Projectile class to spawn */
+	UPROPERTY(EditDefaultsOnly, Category=Projectile)
+	TSubclassOf<class ATP_FirstPersonProjectile> ProjectileClass;
+
+	/** Sound to play each time we fire */
+	UPROPERTY(EditAnywhere, BlueprintReadWrite, Category=Gameplay)
+	USoundBase* FireSound;
+	
+	/** AnimMontage to play each time we fire */
+	UPROPERTY(EditAnywhere, BlueprintReadWrite, Category = Gameplay)
+	UAnimMontage* FireAnimation;
+
+	/** Gun muzzle's offset from the characters location */
+	UPROPERTY(EditAnywhere, BlueprintReadWrite, Category=Gameplay)
+	FVector MuzzleOffset;
+
+	// Sets default values for this component's properties
+	UTP_WeaponComponent();
+
+	// Make the weapon Fire a Projectile
+	UFUNCTION(BlueprintCallable, Category="Weapon")
+	void Fire();
+
+protected:
+	// Called when the game starts
+	virtual void BeginPlay() override;
+
+	// Code for when something overlaps this component
+	UFUNCTION()
+	void OnSphereBeginOverlap(UPrimitiveComponent* OverlappedComponent, AActor* OtherActor, UPrimitiveComponent* OtherComp, int32 OtherBodyIndex, bool bFromSweep, const FHitResult& SweepResult);
+
+private:
+	/** The Character holding this weapon*/
+	ATP_FirstPersonCharacter* Character;
+};
