First Person Templates - Updating comments and template selection description

diff --git a/Templates/TP_FirstPerson/Config/TemplateDefs.ini b/Templates/TP_FirstPerson/Config/TemplateDefs.ini
--- a/Templates/TP_FirstPerson/Config/TemplateDefs.ini
+++ b/Templates/TP_FirstPerson/Config/TemplateDefs.ini
@@ -3,7 +3,7 @@
 Categories=Games
 
 LocalizedDisplayNames=(Language="en",Text="First Person")
-LocalizedDescriptions=(Language="en",Text="The first person template features a player character (represented by a pair of arms) which is viewed from first person perspective. The character can be moved around the level using keyboard, controller or virtual joystick on a touch device. Additionally the player can look around using mouse, controller or virtual joystick on a touch device. The character is also equipped with a gun that, using mouse,  controller or virtual joystick on a touch device will fire a simple sphere projectile that will affect some physics objects in the level, whilst bounce off the arena walls.")
+LocalizedDescriptions=(Language="en",Text="The first person template features a player character (represented by a pair of arms) which is viewed from first person perspective. The character can be moved around the level using keyboard, controller or virtual joystick on a touch device. Additionally the player can look around using mouse, controller or virtual joystick on a touch device. The character can also pick up a gun that, using mouse, controller or virtual joystick on a touch device will fire a simple sphere projectile that will affect some physics objects in the level, whilst bounce off the arena walls.")
 
 LocalizedDisplayNames=(Language="ko",Text="일인칭")
 LocalizedDescriptions=(Language="ko",Text="일인칭 템플릿은 일인칭 시점에서 보이는 (양팔로 표현되는) 플레이어 캐릭터를 특징으로 합니다. 캐릭터 이동은 키보드, 컨트롤러, 터치 디바이스의 가상 조이스틱으로 가능합니다. 추가적으로 플레이어는 마우스, 컨트롤러, 터치 디바이스의 가상 조이스틱을 사용하여 주변을 둘러볼 수도 있습니다. 캐릭터에는 총이 장착되어 있어 마우스, 컨트롤러, 터치 디바이스의 가상 조이스틱을 사용하면 간단한 발사체가 발사되어, 아레나 벽에 튕기면서 레벨에 있는 몇몇 물리 오브젝트에 영향을 줍니다.")
@@ -14,8 +14,8 @@ LocalizedDescriptions=(Language="ja",Text="ファーストパーソンテンプ
 LocalizedDisplayNames=(Language="zh-Hans",Text="第一人称游戏")
 LocalizedDescriptions=(Language="zh-Hans",Text="第一人称模板拥有一个第一人称的玩家角色（通过一双手臂呈现）。您可以通过键盘、手柄或触摸式设备上的虚拟摇杆控制角色在关卡中四处移动。此外，您还可以使用鼠标、手柄或触摸式设备上的虚拟摇杆转动角色视角，观察四周。角色还装备了枪支，可使用鼠标、手柄或触摸式设备上的虚拟摇杆发射一颗球形子弹，并影响关卡中的一些物理对象，同时从竞技场的墙壁上弹开。")
 
-ClassTypes=Character, GameMode, HUD, Actor
-AssetTypes=Animation Sequence, Animation Montage, Sound Wave, Skeletal Mesh, Animation Blueprint, Skeleton, Static Mesh, Texture.
+ClassTypes=Actor, Character, GameMode
+AssetTypes=Animation Sequence, Animation Montage, Sound Wave, Skeletal Mesh, Animation Blueprint, Skeleton, Static Mesh.
 
 FoldersToIgnore=Binaries
 FoldersToIgnore=Build
diff --git a/Templates/TP_FirstPerson/manifest.json b/Templates/TP_FirstPerson/manifest.json
--- a/Templates/TP_FirstPerson/manifest.json
+++ b/Templates/TP_FirstPerson/manifest.json
@@ -23,7 +23,7 @@
 	[
 		{
 			"Language": "en",
-			"Text": "The first person template features a player character (represented by a pair of arms) which is viewed from first person perspective. The character can be moved around the level using keyboard, controller or virtual joystick on a touch device. Additionally the player can look around using mouse, controller or virtual joystick on a touch device. The character is also equipped with a gun that, using mouse,  controller or virtual joystick on a touch device will fire a simple sphere projectile that will affect some physics objects in the level, whilst bounce off the arena walls."
+			"Text": "The first person template features a player character (represented by a pair of arms) which is viewed from first person perspective. The character can be moved around the level using keyboard, controller or virtual joystick on a touch device. Additionally the player can look around using mouse, controller or virtual joystick on a touch device. The character can also pick up a gun that, using mouse, controller or virtual joystick on a touch device will fire a simple sphere projectile that will affect some physics objects in the level, whilst bounce off the arena walls."
 		},
 		{
 			"Language": "ko",
@@ -42,7 +42,7 @@
 	[
 		{
 			"Language": "en",
-			"Text": "Animation Sequence, Animation Montage, Sound Wave, Skeletal Mesh, Animation Blueprint, Skeleton, Static Mesh, Texture."
+			"Text": "Animation Sequence, Animation Montage, Sound Wave, Skeletal Mesh, Animation Blueprint, Skeleton, Static Mesh."
 		}
 	],	
 	"SearchTags":
@@ -52,7 +52,7 @@
 			"Text": "First Person,Shoot,Fire,Aim"
 		}
 	],
-	"ClassTypes": "Character, GameMode, HUD, Actor, Character, GameMode, HUD, Actor",
+	"ClassTypes": "Actor, Character, GameMode",
 	"Category": "CodeFeature",
 	"FocusAsset": "/Game/FirstPerson/Blueprints/FirstPersonCharacter.FirstPersonCharacter",
 	"Thumbnail": "TP_FirstPerson.png",
diff --git a/Templates/TP_FirstPersonBP/manifest.json b/Templates/TP_FirstPersonBP/manifest.json
--- a/Templates/TP_FirstPersonBP/manifest.json
+++ b/Templates/TP_FirstPersonBP/manifest.json
@@ -23,7 +23,7 @@
 	[
 		{
 			"Language": "en",
-			"Text": "The first person template features a player character (represented by a pair of arms) which is viewed from first person perspective. The character can be moved around the level using keyboard, controller or virtual joystick on a touch device. Additionally the player can look around using mouse, controller or virtual joystick on a touch device. The character is also equipped with a gun that, using mouse,  controller or virtual joystick on a touch device will fire a simple sphere projectile that will affect some physics objects in the level, whilst bounce off the arena walls."
+			"Text": "The first person template features a player character (represented by a pair of arms) which is viewed from first person perspective. The character can be moved around the level using keyboard, controller or virtual joystick on a touch device. Additionally the player can look around using mouse, controller or virtual joystick on a touch device. The character can also pick up a gun that, using mouse, controller or virtual joystick on a touch device will fire a simple sphere projectile that will affect some physics objects in the level, whilst bounce off the arena walls."
 		},
 		{
 			"Language": "ko",
@@ -42,7 +42,7 @@
 	[
 		{
 			"Language": "en",
-			"Text": "Animation Sequence, Animation Montage, Sound Wave, Skeletal Mesh, Animation Blueprint, Skeleton, Static Mesh, Texture."
+			"Text": "Animation Sequence, Animation Montage, Sound Wave, Skeletal Mesh, Animation Blueprint, Skeleton, Static Mesh."
 		}
 	],	
 	"SearchTags":
@@ -52,7 +52,7 @@
 			"Text": "First Person,Shoot,Fire,Aim"
 		}
 	],
-	"ClassTypes": "Character, GameMode, HUD, Actor, Character, GameMode, HUD, Actor",
+	"ClassTypes": "Actor, Character, GameMode",
 	"Category": "BlueprintFeature",
 	"FocusAsset": "/Game/FirstPersonBP/FirstPersonOverview.FirstPersonOverview",
 	"Thumbnail": "TP_FirstPersonBP.png",
