Updated template description translation

diff --git a/Templates/TP_FirstPerson/Config/TemplateDefs.ini b/Templates/TP_FirstPerson/Config/TemplateDefs.ini
--- a/Templates/TP_FirstPerson/Config/TemplateDefs.ini
+++ b/Templates/TP_FirstPerson/Config/TemplateDefs.ini
@@ -9,10 +9,10 @@ LocalizedDisplayNames=(Language="ko",Text="일인칭")
 LocalizedDescriptions=(Language="ko",Text="일인칭 템플릿은 일인칭 시점에서 보이는 (양팔로 표현되는) 플레이어 캐릭터를 특징으로 합니다. 캐릭터 이동은 키보드, 컨트롤러, 터치 디바이스의 가상 조이스틱으로 가능합니다. 추가적으로 플레이어는 마우스, 컨트롤러, 터치 디바이스의 가상 조이스틱을 사용하여 주변을 둘러볼 수도 있습니다. 캐릭터에는 총이 장착되어 있어 마우스, 컨트롤러, 터치 디바이스의 가상 조이스틱을 사용하면 간단한 발사체가 발사되어, 아레나 벽에 튕기면서 레벨에 있는 몇몇 물리 오브젝트에 영향을 줍니다.")
 
 LocalizedDisplayNames=(Language="ja",Text="ファーストパーソン")
-LocalizedDescriptions=(Language="ja",Text="ファーストパーソンテンプレートには、一人称視点で表示されるプレイヤーキャラクタ（腕として表示されています）がいます。キーボード、コントローラー、タッチデバイスでのバーチャルジョイスティックでこのキャラクタを移動させることができます。追加して、プレイヤーはマウス、コントローラー、タッチデバイスでのバーチャルジョイスティックで周りを見回すことができます。さらにキャラクタは銃を装備して、マウス、コントローラー、タッチデバイスでのバーチャルジョイスティックを使ってシンプルな球体の弾を発射できます。この弾は物理オブジェクトに影響し、アリーナの壁で反射します。")
+LocalizedDescriptions=(Language="ja",Text="ファーストパーソン テンプレートは、一人称視点から見たプレイヤー キャラクター (両腕が表示される) を特徴とします。タッチ デバイスでキーボード、コントローラー、またはバーチャル ジョイスティックを使用して、レベル内でキャラクターを移動させることができます。さらにプレイヤーは、タッチデバイスでマウス、コントローラー、またはバーチャル ジョイスティックを使って見回すことができます。キャラクターは、タッチデバイスで、マウス、コントローラー、またはバーチャル ジョイスティックを使用して、アリーナの壁で跳ね返りながら、レベル内の一部の物理オブジェクトに影響を与えるシンプルな球体の発射体を発射する銃を拾うこともできます。")
 
 LocalizedDisplayNames=(Language="zh-Hans",Text="第一人称游戏")
-LocalizedDescriptions=(Language="zh-Hans",Text="第一人称模板拥有一个第一人称的玩家角色（通过一双手臂呈现）。您可以通过键盘、手柄或触摸式设备上的虚拟摇杆控制角色在关卡中四处移动。此外，您还可以使用鼠标、手柄或触摸式设备上的虚拟摇杆转动角色视角，观察四周。角色还装备了枪支，可使用鼠标、手柄或触摸式设备上的虚拟摇杆发射一颗球形子弹，并影响关卡中的一些物理对象，同时从竞技场的墙壁上弹开。")
+LocalizedDescriptions=(Language="zh-Hans",Text="第一人称模板包含一个第一人称玩家角色（通过一双手臂呈现）。您可以通过键盘、手柄或触摸式设备上的虚拟摇杆控制角色在关卡中四处移动。你还可以使用鼠标、手柄或触摸式设备上的虚拟摇杆转动角色视角，观察四周。此外，你还可以让角色拾取枪支，然后用鼠标、手柄或触摸式设备上的虚拟摇杆来发射球形子弹。子弹会撞飞关卡中的物理对象，还能从竞技场墙壁上弹开。")
 
 ClassTypes=Actor, Character, GameMode
 AssetTypes=Animation Sequence, Animation Montage, Sound Wave, Skeletal Mesh, Animation Blueprint, Skeleton, Static Mesh.
diff --git a/Templates/TP_FirstPerson/manifest.json b/Templates/TP_FirstPerson/manifest.json
--- a/Templates/TP_FirstPerson/manifest.json
+++ b/Templates/TP_FirstPerson/manifest.json
@@ -31,11 +31,11 @@
 		},
 		{
 			"Language": "zh-Hans",
-			"Text": "第一人称模板拥有一个第一人称的玩家角色（通过一双手臂呈现）。您可以通过键盘、手柄或触摸式设备上的虚拟摇杆控制角色在关卡中四处移动。此外，您还可以使用鼠标、手柄或触摸式设备上的虚拟摇杆转动角色视角，观察四周。角色还装备了枪支，可使用鼠标、手柄或触摸式设备上的虚拟摇杆发射一颗球形子弹，并影响关卡中的一些物理对象，同时从竞技场的墙壁上弹开。"
+			"Text": "第一人称模板包含一个第一人称玩家角色（通过一双手臂呈现）。您可以通过键盘、手柄或触摸式设备上的虚拟摇杆控制角色在关卡中四处移动。你还可以使用鼠标、手柄或触摸式设备上的虚拟摇杆转动角色视角，观察四周。此外，你还可以让角色拾取枪支，然后用鼠标、手柄或触摸式设备上的虚拟摇杆来发射球形子弹。子弹会撞飞关卡中的物理对象，还能从竞技场墙壁上弹开。"
 		},
 		{
 			"Language": "ja",
-			"Text": "ファーストパーソンテンプレートには、一人称視点で表示されるプレイヤーキャラクタ（腕として表示されています）がいます。キーボード、コントローラー、タッチデバイスでのバーチャルジョイスティックでこのキャラクタを移動させることができます。追加して、プレイヤーはマウス、コントローラー、タッチデバイスでのバーチャルジョイスティックで周りを見回すことができます。さらにキャラクタは銃を装備して、マウス、コントローラー、タッチデバイスでのバーチャルジョイスティックを使ってシンプルな球体の弾を発射できます。この弾は物理オブジェクトに影響し、アリーナの壁で反射します。"
+			"Text": "ファーストパーソン テンプレートは、一人称視点から見たプレイヤー キャラクター (両腕が表示される) を特徴とします。タッチ デバイスでキーボード、コントローラー、またはバーチャル ジョイスティックを使用して、レベル内でキャラクターを移動させることができます。さらにプレイヤーは、タッチデバイスでマウス、コントローラー、またはバーチャル ジョイスティックを使って見回すことができます。キャラクターは、タッチデバイスで、マウス、コントローラー、またはバーチャル ジョイスティックを使用して、アリーナの壁で跳ね返りながら、レベル内の一部の物理オブジェクトに影響を与えるシンプルな球体の発射体を発射する銃を拾うこともできます。"
 		}
 	],
 	"AssetTypes":
diff --git a/Templates/TP_ThirdPerson/Config/TemplateDefs.ini b/Templates/TP_ThirdPerson/Config/TemplateDefs.ini
--- a/Templates/TP_ThirdPerson/Config/TemplateDefs.ini
+++ b/Templates/TP_ThirdPerson/Config/TemplateDefs.ini
@@ -10,10 +10,10 @@ LocalizedDisplayNames=(Language="ko",Text="삼인칭")
 LocalizedDescriptions=(Language="ko",Text="삼인칭 템플릿은 카메라가 머리위 약간 뒤쪽에 위치한 플레이가능 캐릭터를 특징으로 합니다. 마우스, 키보드, 터치 디바이스의 가상 조이스틱으로 캐릭터가 이동함에 따라 카메라는 캐릭터를 따라다니며 어깨위 시점으로 봅니다. 주인공 캐릭터를 강조하는 시점으로, 액션이나 어드벤처 게임에 많이 쓰입니다. 레벨에는 플레이어가 점프할 수 있는 오브젝트가 여럿 있습니다.")
 
 LocalizedDisplayNames=(Language="ja",Text="サードパーソン")
-LocalizedDescriptions=(Language="ja",Text="サードパーソンテンプレートにはプレイアブルキャラクタが存在し、カメラがキャラクタの背後の少し上方にあります。キーボード、コントローラー、タッチデバイスでのバーチャルジョイスティックでキャラクタが移動すると、カメラは肩越しのビューで追従します。この視点はメインキャラクタを強調するのもでアクションアドベンチャーゲームで人気があるものです。レベルにはプレイヤーがジャンプして乗ることができるオブジェクトが置かれています。")
+LocalizedDescriptions=(Language="ja",Text="サードパーソン パックは、後ろ側と少し上にカメラが配置されているプレイ可能なキャラクターを備えています。タッチデバイスでマウス、キーボード、コントローラー、またはバーチャル ジョイスティックの使用によりキャラクターが移動すると、カメラはキャラクターを追跡し、肩越しの視点から観察します。この視点では主人公が強調されるため、アクションゲームやアドベンチャー ゲームで人気があります。テンプレート マップには、プレイヤーがジャンプしたり押すことができるオブジェクトがいくつか含まれています。")
 
 LocalizedDisplayNames=(Language="zh-Hans",Text="第三人称游戏")
-LocalizedDescriptions=(Language="zh-Hans",Text="第三人称游戏模板拥有一个可游玩的角色，摄像机位于角色身后偏上的位置。您可以通过鼠标、键盘、手柄或触摸式设备上控制角色移动，摄像机会跟随角色，通过过肩式视角观察世界。此视角突出主角，在动作及冒险类游戏中非常流行。关卡中还包含一些对象，可以让角色跳到上面。")
+LocalizedDescriptions=(Language="zh-Hans",Text="第三人称游戏内容包包含一个可操控的角色以及一台位于角色身后斜上方的摄像机。你可以通过鼠标、键盘、手柄或触摸式设备上控制角色移动。摄像机会跟随角色，以越肩视角拍摄场景。这种视角突出主角，常用于动作及冒险类游戏。该模板的地图中还包含一些对象，可以让角色跳上跳下或四处推动。")
 
 ClassTypes=Character, SpringArmComponent, GameMode
 AssetTypes=Animation Sequence, Animation Blueprint, Skeleton, Skeletal Mesh
diff --git a/Templates/TP_ThirdPerson/manifest.json b/Templates/TP_ThirdPerson/manifest.json
--- a/Templates/TP_ThirdPerson/manifest.json
+++ b/Templates/TP_ThirdPerson/manifest.json
@@ -31,11 +31,11 @@
 		},
 		{
 			"Language": "zh-Hans",
-			"Text": "第三人称游戏模板拥有一个可游玩的角色，摄像机位于角色身后偏上的位置。您可以通过鼠标、键盘、手柄或触摸式设备上控制角色移动，摄像机会跟随角色，通过过肩式视角观察世界。此视角突出主角，在动作及冒险类游戏中非常流行。关卡中还包含一些对象，可以让角色跳到上面。"
+			"Text": "第三人称游戏内容包包含一个可操控的角色以及一台位于角色身后斜上方的摄像机。你可以通过鼠标、键盘、手柄或触摸式设备上控制角色移动。摄像机会跟随角色，以越肩视角拍摄场景。这种视角突出主角，常用于动作及冒险类游戏。该模板的地图中还包含一些对象，可以让角色跳上跳下或四处推动。"
 		},
 		{
 			"Language": "ja",
-			"Text": "サードパーソンテンプレートにはプレイアブルキャラクタが存在し、カメラがキャラクタの背後の少し上方にあります。キーボード、コントローラー、タッチデバイスでのバーチャルジョイスティックでキャラクタが移動すると、カメラは肩越しのビューで追従します。この視点はメインキャラクタを強調するのもでアクションアドベンチャーゲームで人気があるものです。レベルにはプレイヤーがジャンプして乗ることができるオブジェクトが置かれています。"
+			"Text": "サードパーソン パックは、後ろ側と少し上にカメラが配置されているプレイ可能なキャラクターを備えています。タッチデバイスでマウス、キーボード、コントローラー、またはバーチャル ジョイスティックの使用によりキャラクターが移動すると、カメラはキャラクターを追跡し、肩越しの視点から観察します。この視点では主人公が強調されるため、アクションゲームやアドベンチャー ゲームで人気があります。テンプレート マップには、プレイヤーがジャンプしたり押すことができるオブジェクトがいくつか含まれています。"
 		}
 	],
 	"AssetTypes":
diff --git a/Templates/TP_ThirdPersonBP/Config/TemplateDefs.ini b/Templates/TP_ThirdPersonBP/Config/TemplateDefs.ini
--- a/Templates/TP_ThirdPersonBP/Config/TemplateDefs.ini
+++ b/Templates/TP_ThirdPersonBP/Config/TemplateDefs.ini
@@ -9,10 +9,10 @@ LocalizedDisplayNames=(Language="ko",Text="삼인칭")
 LocalizedDescriptions=(Language="ko",Text="삼인칭 템플릿은 카메라가 머리위 약간 뒤쪽에 위치한 플레이가능 캐릭터를 특징으로 합니다. 마우스, 키보드, 터치 디바이스의 가상 조이스틱으로 캐릭터가 이동함에 따라 카메라는 캐릭터를 따라다니며 어깨위 시점으로 봅니다. 주인공 캐릭터를 강조하는 시점으로, 액션이나 어드벤처 게임에 많이 쓰입니다. 레벨에는 플레이어가 점프할 수 있는 오브젝트가 여럿 있습니다.")
 
 LocalizedDisplayNames=(Language="ja",Text="サードパーソン")
-LocalizedDescriptions=(Language="ja",Text="サードパーソンテンプレートにはプレイアブルキャラクタが存在し、カメラがキャラクタの背後の少し上方にあります。キーボード、コントローラー、タッチデバイスでのバーチャルジョイスティックでキャラクタが移動すると、カメラは肩越しのビューで追従します。この視点はメインキャラクタを強調するのもでアクションアドベンチャーゲームで人気があるものです。レベルにはプレイヤーがジャンプして乗ることができるオブジェクトが置かれています。")
+LocalizedDescriptions=(Language="ja",Text="サードパーソン パックは、後ろ側と少し上にカメラが配置されているプレイ可能なキャラクターを備えています。タッチデバイスでマウス、キーボード、コントローラー、またはバーチャル ジョイスティックの使用によりキャラクターが移動すると、カメラはキャラクターを追跡し、肩越しの視点から観察します。この視点では主人公が強調されるため、アクションゲームやアドベンチャー ゲームで人気があります。テンプレート マップには、プレイヤーがジャンプしたり押すことができるオブジェクトがいくつか含まれています。")
 
 LocalizedDisplayNames=(Language="zh-Hans",Text="第三人称游戏")
-LocalizedDescriptions=(Language="zh-Hans",Text="第三人称游戏模板拥有一个可游玩的角色，摄像机位于角色身后偏上的位置。您可以通过鼠标、键盘、手柄或触摸式设备上控制角色移动，摄像机会跟随角色，通过过肩式视角观察世界。此视角突出主角，在动作及冒险类游戏中非常流行。关卡中还包含一些对象，可以让角色跳到上面。")
+LocalizedDescriptions=(Language="zh-Hans",Text="第三人称游戏内容包包含一个可操控的角色以及一台位于角色身后斜上方的摄像机。你可以通过鼠标、键盘、手柄或触摸式设备上控制角色移动。摄像机会跟随角色，以越肩视角拍摄场景。这种视角突出主角，常用于动作及冒险类游戏。该模板的地图中还包含一些对象，可以让角色跳上跳下或四处推动。")
 
 ClassTypes=Character, SpringArmComponent, GameMode
 AssetTypes=Animation Sequence, Animation Blueprint, Skeleton, Skeletal Mesh
diff --git a/Templates/TP_ThirdPersonBP/manifest.json b/Templates/TP_ThirdPersonBP/manifest.json
--- a/Templates/TP_ThirdPersonBP/manifest.json
+++ b/Templates/TP_ThirdPersonBP/manifest.json
@@ -31,11 +31,11 @@
 		},
 		{
 			"Language": "zh-Hans",
-			"Text": "第三人称游戏模板拥有一个可游玩的角色，摄像机位于角色身后偏上的位置。您可以通过鼠标、键盘、手柄或触摸式设备上控制角色移动，摄像机会跟随角色，通过过肩式视角观察世界。此视角突出主角，在动作及冒险类游戏中非常流行。关卡中还包含一些对象，可以让角色跳到上面。"
+			"Text": "第三人称游戏内容包包含一个可操控的角色以及一台位于角色身后斜上方的摄像机。你可以通过鼠标、键盘、手柄或触摸式设备上控制角色移动。摄像机会跟随角色，以越肩视角拍摄场景。这种视角突出主角，常用于动作及冒险类游戏。该模板的地图中还包含一些对象，可以让角色跳上跳下或四处推动。"
 		},
 		{
 			"Language": "ja",
-			"Text": "サードパーソンテンプレートにはプレイアブルキャラクタが存在し、カメラがキャラクタの背後の少し上方にあります。キーボード、コントローラー、タッチデバイスでのバーチャルジョイスティックでキャラクタが移動すると、カメラは肩越しのビューで追従します。この視点はメインキャラクタを強調するのもでアクションアドベンチャーゲームで人気があるものです。レベルにはプレイヤーがジャンプして乗ることができるオブジェクトが置かれています。"
+			"Text": "サードパーソン パックは、後ろ側と少し上にカメラが配置されているプレイ可能なキャラクターを備えています。タッチデバイスでマウス、キーボード、コントローラー、またはバーチャル ジョイスティックの使用によりキャラクターが移動すると、カメラはキャラクターを追跡し、肩越しの視点から観察します。この視点では主人公が強調されるため、アクションゲームやアドベンチャー ゲームで人気があります。テンプレート マップには、プレイヤーがジャンプしたり押すことができるオブジェクトがいくつか含まれています。"
 		}
 	],
 	"AssetTypes":
diff --git a/Templates/TP_TopDownBP/Config/TemplateDefs.ini b/Templates/TP_TopDownBP/Config/TemplateDefs.ini
--- a/Templates/TP_TopDownBP/Config/TemplateDefs.ini
+++ b/Templates/TP_TopDownBP/Config/TemplateDefs.ini
@@ -10,10 +10,10 @@ LocalizedDisplayNames=(Language="ko",Text="내려보기")
 LocalizedDescriptions=(Language="ko",Text="내려보기 템플릿은 플레이어 캐릭터 위 삼인칭 시점보다 멀리 카메라가 위치합니다. 플레이어 제어는 마우스나 터치스크린으로 원하는 목적지를 클릭하는 식으로 이루어지며, 원활한 캐릭터 이동을 위해 내비게이션 시스템을 사용합니다. 이 시점은 액션 롤플레잉 게임에 자주 쓰입니다.")
 
 LocalizedDisplayNames=(Language="ja",Text="トップダウン")
-LocalizedDescriptions=(Language="ja",Text="トップダウンテンプレートにはプレイアブルキャラクタが存在し、カメラはキャラクタの背後、サードパーソンテンプレートより遠い上方にあります。マウスやタッチスクリーンで行き先をクリックすることでキャラクタを操作することができます。キャラクタはナビゲーションシステムを利用して移動します。この視点はアクションロールプレイングゲームでよく使用されています。")
+LocalizedDescriptions=(Language="ja",Text="トップダウン テンプレートは、カメラがアバターの後ろと上に配置されている (サードパーソン視点よりも距離が遠い位置) プレイ可能なキャラクターを特徴としています。プレイヤーは、マウス、ゲームパッド、またはタッチ スクリーンの使用により制御され、目的地をクリックするか、入力が保持されている場合はカーソルをたどることでキャラクターを移動させます。この視点は、アクション ロールプレイ ゲームでよく使用されます。")
 
 LocalizedDisplayNames=(Language="zh-Hans",Text="俯视角游戏")
-LocalizedDescriptions=(Language="zh-Hans",Text="俯视角游戏模板拥有一个可游玩角色，摄像机位于角色的后上方，但位置比第三人称视角要高得多。玩家可以使用鼠标或触摸屏点击目的地，或使用导航系统来辅助角色移动。此视角常用于动作角色扮演类游戏。")
+LocalizedDescriptions=(Language="zh-Hans",Text="俯视角游戏模板包含一个可操控的角色以及一台位于角色后上方的摄像机。摄像机的位置比第三人称视角模板中的要高很多。玩家可以通过点击鼠标、触控板或触摸屏来选择目的地并移动角色，或持续按下鼠标来移动角色。此视角常用于动作角色扮演类游戏。")
 
 ClassTypes=Character, PlayerController, CameraComponent, SpringArmComponent, GameMode
 AssetTypes=Animation Sequence, Skeletal Mesh, Animation Blueprint,Skeleton
diff --git a/Templates/TP_TopDownBP/manifest.json b/Templates/TP_TopDownBP/manifest.json
--- a/Templates/TP_TopDownBP/manifest.json
+++ b/Templates/TP_TopDownBP/manifest.json
@@ -31,11 +31,11 @@
 		},
 		{
 			"Language": "zh-Hans",
-			"Text": "俯视角游戏模板拥有一个可游玩角色，摄像机位于角色的后上方，但位置比第三人称视角要高得多。玩家可以使用鼠标或触摸屏点击目的地，或使用导航系统来辅助角色移动。此视角常用于动作角色扮演类游戏。"
+			"Text": "俯视角游戏模板包含一个可操控的角色以及一台位于角色后上方的摄像机。摄像机的位置比第三人称视角模板中的要高很多。玩家可以通过点击鼠标、触控板或触摸屏来选择目的地并移动角色，或持续按下鼠标来移动角色。此视角常用于动作角色扮演类游戏。"
 		},
 		{
 			"Language": "ja",
-			"Text": "トップダウンテンプレートにはプレイアブルキャラクタが存在し、カメラはキャラクタの背後、サードパーソンテンプレートより遠い上方にあります。マウスやタッチスクリーンで行き先をクリックすることでキャラクタを操作することができます。キャラクタはナビゲーションシステムを利用して移動します。この視点はアクションロールプレイングゲームでよく使用されています。"
+			"Text": "トップダウン テンプレートは、カメラがアバターの後ろと上に配置されている (サードパーソン視点よりも距離が遠い位置) プレイ可能なキャラクターを特徴としています。プレイヤーは、マウス、ゲームパッド、またはタッチ スクリーンの使用により制御され、目的地をクリックするか、入力が保持されている場合はカーソルをたどることでキャラクターを移動させます。この視点は、アクション ロールプレイ ゲームでよく使用されます。"
 		}
 	],
 	"AssetTypes":
