Updated template description translations (Korean)

diff --git a/Templates/TP_FirstPerson/Config/TemplateDefs.ini b/Templates/TP_FirstPerson/Config/TemplateDefs.ini
--- a/Templates/TP_FirstPerson/Config/TemplateDefs.ini
+++ b/Templates/TP_FirstPerson/Config/TemplateDefs.ini
@@ -6,7 +6,7 @@ LocalizedDisplayNames=(Language="en",Text="First Person")
 LocalizedDescriptions=(Language="en",Text="The first person template features a player character (represented by a pair of arms) which is viewed from first person perspective. The character can be moved around the level using keyboard, controller or virtual joystick on a touch device. Additionally the player can look around using mouse, controller or virtual joystick on a touch device. The character can also pick up a gun that, using mouse, controller or virtual joystick on a touch device will fire a simple sphere projectile that will affect some physics objects in the level, whilst bounce off the arena walls.")
 
 LocalizedDisplayNames=(Language="ko",Text="일인칭")
-LocalizedDescriptions=(Language="ko",Text="일인칭 템플릿은 일인칭 시점에서 보이는 (양팔로 표현되는) 플레이어 캐릭터를 특징으로 합니다. 캐릭터 이동은 키보드, 컨트롤러, 터치 디바이스의 가상 조이스틱으로 가능합니다. 추가적으로 플레이어는 마우스, 컨트롤러, 터치 디바이스의 가상 조이스틱을 사용하여 주변을 둘러볼 수도 있습니다. 캐릭터에는 총이 장착되어 있어 마우스, 컨트롤러, 터치 디바이스의 가상 조이스틱을 사용하면 간단한 발사체가 발사되어, 아레나 벽에 튕기면서 레벨에 있는 몇몇 물리 오브젝트에 영향을 줍니다.")
+LocalizedDescriptions=(Language="ko",Text="일인칭 템플릿은 일인칭 시점에서 보이는 (양팔로 표현되는) 플레이어 캐릭터를 특징으로 합니다. 캐릭터는 레벨 내에서 키보드, 컨트롤러나 터치 디바이스의 가상 조이스틱으로 이동할 수 있습니다. 추가적으로 플레이어는 마우스, 컨트롤러나 터치 디바이스의 가상 조이스틱을 사용하여 주변을 둘러볼 수도 있습니다. 또한, 캐릭터는 총을 장착할 수 있으며 마우스, 컨트롤러나 터치 디바이스의 가상 조이스틱을 사용해 간단한 발사체를 발사하여, 아레나 벽에 튕기면서 레벨에 있는 일부 피직스 오브젝트에 영향을 줍니다.")
 
 LocalizedDisplayNames=(Language="ja",Text="ファーストパーソン")
 LocalizedDescriptions=(Language="ja",Text="ファーストパーソン テンプレートは、一人称視点から見たプレイヤー キャラクター (両腕が表示される) を特徴とします。タッチ デバイスでキーボード、コントローラー、またはバーチャル ジョイスティックを使用して、レベル内でキャラクターを移動させることができます。さらにプレイヤーは、タッチデバイスでマウス、コントローラー、またはバーチャル ジョイスティックを使って見回すことができます。キャラクターは、タッチデバイスで、マウス、コントローラー、またはバーチャル ジョイスティックを使用して、アリーナの壁で跳ね返りながら、レベル内の一部の物理オブジェクトに影響を与えるシンプルな球体の発射体を発射する銃を拾うこともできます。")
diff --git a/Templates/TP_ThirdPerson/Config/TemplateDefs.ini b/Templates/TP_ThirdPerson/Config/TemplateDefs.ini
--- a/Templates/TP_ThirdPerson/Config/TemplateDefs.ini
+++ b/Templates/TP_ThirdPerson/Config/TemplateDefs.ini
@@ -7,7 +7,7 @@ LocalizedDisplayNames=(Language="en",Text="Third Person")
 LocalizedDescriptions=(Language="en",Text="The third person pack features a playable character where the camera is positioned behind and slightly above the character. As the character moves using mouse, keyboard, controller or virtual joystick on a touch device, the camera follows the character, observing from an over-the-shoulder point of view. This perspective emphasizes the main character and is popular in action and adventure games. The template map contains several objects that the player can jump upon, and push around.")
 
 LocalizedDisplayNames=(Language="ko",Text="삼인칭")
-LocalizedDescriptions=(Language="ko",Text="삼인칭 템플릿은 카메라가 머리위 약간 뒤쪽에 위치한 플레이가능 캐릭터를 특징으로 합니다. 마우스, 키보드, 터치 디바이스의 가상 조이스틱으로 캐릭터가 이동함에 따라 카메라는 캐릭터를 따라다니며 어깨위 시점으로 봅니다. 주인공 캐릭터를 강조하는 시점으로, 액션이나 어드벤처 게임에 많이 쓰입니다. 레벨에는 플레이어가 점프할 수 있는 오브젝트가 여럿 있습니다.")
+LocalizedDescriptions=(Language="ko",Text="삼인칭 템플릿에는 플레이 가능한 캐릭터가 포함되어 있으며 카메라가 캐릭터 머리 위 약간 뒤쪽에 위치합니다. 마우스, 키보드, 터치 디바이스의 가상 조이스틱으로 캐릭터가 이동함에 따라 카메라가 따라 움직이며 어깨 위 시점으로 봅니다. 이 시점은 주인공 캐릭터를 강조해 액션이나 어드벤처 게임에 주로 활용됩니다. 이 템플릿 맵에는 플레이어가 점프하여 착지하고 밀 수 있는 여러가지 오브젝트가 포함되어 있습니다.")
 
 LocalizedDisplayNames=(Language="ja",Text="サードパーソン")
 LocalizedDescriptions=(Language="ja",Text="サードパーソン パックは、後ろ側と少し上にカメラが配置されているプレイ可能なキャラクターを備えています。タッチデバイスでマウス、キーボード、コントローラー、またはバーチャル ジョイスティックの使用によりキャラクターが移動すると、カメラはキャラクターを追跡し、肩越しの視点から観察します。この視点では主人公が強調されるため、アクションゲームやアドベンチャー ゲームで人気があります。テンプレート マップには、プレイヤーがジャンプしたり押すことができるオブジェクトがいくつか含まれています。")
diff --git a/Templates/TP_ThirdPersonBP/Config/TemplateDefs.ini b/Templates/TP_ThirdPersonBP/Config/TemplateDefs.ini
--- a/Templates/TP_ThirdPersonBP/Config/TemplateDefs.ini
+++ b/Templates/TP_ThirdPersonBP/Config/TemplateDefs.ini
@@ -6,7 +6,7 @@ LocalizedDisplayNames=(Language="en",Text="Third Person")
 LocalizedDescriptions=(Language="en",Text="The third person pack features a playable character where the camera is positioned behind and slightly above the character. As the character moves using mouse, keyboard, controller or virtual joystick on a touch device, the camera follows the character, observing from an over-the-shoulder point of view. This perspective emphasizes the main character and is popular in action and adventure games. The template map contains several objects that the player can jump upon, and push around.")
 
 LocalizedDisplayNames=(Language="ko",Text="삼인칭")
-LocalizedDescriptions=(Language="ko",Text="삼인칭 템플릿은 카메라가 머리위 약간 뒤쪽에 위치한 플레이가능 캐릭터를 특징으로 합니다. 마우스, 키보드, 터치 디바이스의 가상 조이스틱으로 캐릭터가 이동함에 따라 카메라는 캐릭터를 따라다니며 어깨위 시점으로 봅니다. 주인공 캐릭터를 강조하는 시점으로, 액션이나 어드벤처 게임에 많이 쓰입니다. 레벨에는 플레이어가 점프할 수 있는 오브젝트가 여럿 있습니다.")
+LocalizedDescriptions=(Language="ko",Text="삼인칭 템플릿에는 플레이 가능한 캐릭터가 포함되어 있으며 카메라가 캐릭터 머리 위 약간 뒤쪽에 위치합니다. 마우스, 키보드, 터치 디바이스의 가상 조이스틱으로 캐릭터가 이동함에 따라 카메라가 따라 움직이며 어깨 위 시점으로 봅니다. 이 시점은 주인공 캐릭터를 강조해 액션이나 어드벤처 게임에 주로 활용됩니다. 이 템플릿 맵에는 플레이어가 점프하여 착지하고 밀 수 있는 여러가지 오브젝트가 포함되어 있습니다.")
 
 LocalizedDisplayNames=(Language="ja",Text="サードパーソン")
 LocalizedDescriptions=(Language="ja",Text="サードパーソン パックは、後ろ側と少し上にカメラが配置されているプレイ可能なキャラクターを備えています。タッチデバイスでマウス、キーボード、コントローラー、またはバーチャル ジョイスティックの使用によりキャラクターが移動すると、カメラはキャラクターを追跡し、肩越しの視点から観察します。この視点では主人公が強調されるため、アクションゲームやアドベンチャー ゲームで人気があります。テンプレート マップには、プレイヤーがジャンプしたり押すことができるオブジェクトがいくつか含まれています。")
diff --git a/Templates/TP_ThirdPersonBP/manifest.json b/Templates/TP_ThirdPersonBP/manifest.json
--- a/Templates/TP_ThirdPersonBP/manifest.json
+++ b/Templates/TP_ThirdPersonBP/manifest.json
@@ -27,7 +27,7 @@
 		},
 		{
 			"Language": "ko",
-			"Text": "삼인칭 템플릿은 카메라가 머리위 약간 뒤쪽에 위치한 플레이가능 캐릭터를 특징으로 합니다. 마우스, 키보드, 터치 디바이스의 가상 조이스틱으로 캐릭터가 이동함에 따라 카메라는 캐릭터를 따라다니며 어깨위 시점으로 봅니다. 주인공 캐릭터를 강조하는 시점으로, 액션이나 어드벤처 게임에 많이 쓰입니다. 레벨에는 플레이어가 점프할 수 있는 오브젝트가 여럿 있습니다."
+			"Text": "삼인칭 템플릿에는 플레이 가능한 캐릭터가 포함되어 있으며 카메라가 캐릭터 머리 위 약간 뒤쪽에 위치합니다. 마우스, 키보드, 터치 디바이스의 가상 조이스틱으로 캐릭터가 이동함에 따라 카메라가 따라 움직이며 어깨 위 시점으로 봅니다. 이 시점은 주인공 캐릭터를 강조해 액션이나 어드벤처 게임에 주로 활용됩니다. 이 템플릿 맵에는 플레이어가 점프하여 착지하고 밀 수 있는 여러가지 오브젝트가 포함되어 있습니다."
 		},
 		{
 			"Language": "zh-Hans",
diff --git a/Templates/TP_TopDownBP/Config/TemplateDefs.ini b/Templates/TP_TopDownBP/Config/TemplateDefs.ini
--- a/Templates/TP_TopDownBP/Config/TemplateDefs.ini
+++ b/Templates/TP_TopDownBP/Config/TemplateDefs.ini
@@ -7,7 +7,7 @@ LocalizedDisplayNames=(Language="en",Text="Top Down")
 LocalizedDescriptions=(Language="en",Text="The top down template features a playable character where the camera is positioned behind and above the avatar at greater distance than in third person perspective. The player is controlled by using the mouse, gamepad, or touchscreen and moves the character by clicking on a destination, or following the cursor if input is held. This perspective is often used in action role playing games.")
 
 LocalizedDisplayNames=(Language="ko",Text="내려보기")
-LocalizedDescriptions=(Language="ko",Text="내려보기 템플릿은 플레이어 캐릭터 위 삼인칭 시점보다 멀리 카메라가 위치합니다. 플레이어 제어는 마우스나 터치스크린으로 원하는 목적지를 클릭하는 식으로 이루어지며, 원활한 캐릭터 이동을 위해 내비게이션 시스템을 사용합니다. 이 시점은 액션 롤플레잉 게임에 자주 쓰입니다.")
+LocalizedDescriptions=(Language="ko",Text="탑다운 템플릿에는 플레이 가능한 캐릭터가 포함되어 있으며 카메라가 아바타 머리 위 약간 뒤쪽에서 삼인칭 시점보다 더 원거리에 위치합니다. 플레이어는 마우스, 게임패드나 터치스크린으로 제어하고, 원하는 목적지를 클릭하거나 커서의 입력을 유지하여 따라 이동하게 만들 수 있습니다. 이 시점은 주로 액션 RPG에 활용됩니다.")
 
 LocalizedDisplayNames=(Language="ja",Text="トップダウン")
 LocalizedDescriptions=(Language="ja",Text="トップダウン テンプレートは、カメラがアバターの後ろと上に配置されている (サードパーソン視点よりも距離が遠い位置) プレイ可能なキャラクターを特徴としています。プレイヤーは、マウス、ゲームパッド、またはタッチ スクリーンの使用により制御され、目的地をクリックするか、入力が保持されている場合はカーソルをたどることでキャラクターを移動させます。この視点は、アクション ロールプレイ ゲームでよく使用されます。")
