Updated template description translation (Korean)

diff --git a/Templates/TP_FirstPerson/Config/TemplateDefs.ini b/Templates/TP_FirstPerson/Config/TemplateDefs.ini
--- a/Templates/TP_FirstPerson/Config/TemplateDefs.ini
+++ b/Templates/TP_FirstPerson/Config/TemplateDefs.ini
@@ -6,7 +6,7 @@ LocalizedDisplayNames=(Language="en",Text="First Person")
 LocalizedDescriptions=(Language="en",Text="The first person template features a player character (represented by a pair of arms) which is viewed from first person perspective. The character can be moved around the level using keyboard, controller or virtual joystick on a touch device. Additionally the player can look around using mouse, controller or virtual joystick on a touch device. The character can also pick up a gun that, using mouse, controller or virtual joystick on a touch device will fire a simple sphere projectile that will affect some physics objects in the level, whilst bounce off the arena walls.")
 
 LocalizedDisplayNames=(Language="ko",Text="일인칭")
-LocalizedDescriptions=(Language="ko",Text="일인칭 템플릿은 일인칭 시점에서 보이는 (양팔로 표현되는) 플레이어 캐릭터를 특징으로 합니다. 캐릭터는 레벨 내에서 키보드, 컨트롤러나 터치 디바이스의 가상 조이스틱으로 이동할 수 있습니다. 추가적으로 플레이어는 마우스, 컨트롤러나 터치 디바이스의 가상 조이스틱을 사용하여 주변을 둘러볼 수도 있습니다. 또한, 캐릭터는 총을 장착할 수 있으며 마우스, 컨트롤러나 터치 디바이스의 가상 조이스틱을 사용해 간단한 발사체를 발사하여, 아레나 벽에 튕기면서 레벨에 있는 일부 피직스 오브젝트에 영향을 줍니다.")
+LocalizedDescriptions=(Language="ko",Text="일인칭 템플릿은 일인칭 시점에서 보이는 (양팔로 표현되는) 플레이어 캐릭터를 특징으로 합니다. 캐릭터는 레벨 내에서 키보드, 컨트롤러나 터치 디바이스의 가상 조이스틱으로 이동할 수 있습니다. 추가적으로 플레이어는 마우스, 컨트롤러나 터치 디바이스의 가상 조이스틱을 사용하여 주변을 둘러볼 수도 있습니다. 또한, 캐릭터는 총을 장착할 수 있으며 마우스, 컨트롤러나 터치 디바이스의 가상 조이스틱을 사용해 간단한 구형 발사체를 발사하여, 아레나 벽에 튕기면서 레벨에 있는 일부 피직스 오브젝트에 영향을 줍니다.")
 
 LocalizedDisplayNames=(Language="ja",Text="ファーストパーソン")
 LocalizedDescriptions=(Language="ja",Text="ファーストパーソン テンプレートは、一人称視点から見たプレイヤー キャラクター (両腕が表示される) を特徴とします。タッチ デバイスでキーボード、コントローラー、またはバーチャル ジョイスティックを使用して、レベル内でキャラクターを移動させることができます。さらにプレイヤーは、タッチデバイスでマウス、コントローラー、またはバーチャル ジョイスティックを使って見回すことができます。キャラクターは、タッチデバイスで、マウス、コントローラー、またはバーチャル ジョイスティックを使用して、アリーナの壁で跳ね返りながら、レベル内の一部の物理オブジェクトに影響を与えるシンプルな球体の発射体を発射する銃を拾うこともできます。")
