Fixing typos for Chinese template descriptions to ensure consistency, as per

diff --git a/Templates/TP_FirstPerson/Config/TemplateDefs.ini b/Templates/TP_FirstPerson/Config/TemplateDefs.ini
--- a/Templates/TP_FirstPerson/Config/TemplateDefs.ini
+++ b/Templates/TP_FirstPerson/Config/TemplateDefs.ini
@@ -12,7 +12,7 @@ LocalizedDisplayNames=(Language="ja",Text="ファーストパーソン")
 LocalizedDescriptions=(Language="ja",Text="ファーストパーソン テンプレートは、一人称視点から見たプレイヤー キャラクター (両腕が表示される) を特徴とします。タッチ デバイスでキーボード、コントローラー、またはバーチャル ジョイスティックを使用して、レベル内でキャラクターを移動させることができます。さらにプレイヤーは、タッチデバイスでマウス、コントローラー、またはバーチャル ジョイスティックを使って見回すことができます。キャラクターは、タッチデバイスで、マウス、コントローラー、またはバーチャル ジョイスティックを使用して、アリーナの壁で跳ね返りながら、レベル内の一部の物理オブジェクトに影響を与えるシンプルな球体の発射体を発射する銃を拾うこともできます。")
 
 LocalizedDisplayNames=(Language="zh-Hans",Text="第一人称游戏")
-LocalizedDescriptions=(Language="zh-Hans",Text="第一人称模板包含一个第一人称玩家角色（通过一双手臂呈现）。您可以通过键盘、手柄或触摸式设备上的虚拟摇杆控制角色在关卡中四处移动。你还可以使用鼠标、手柄或触摸式设备上的虚拟摇杆转动角色视角，观察四周。此外，你还可以让角色拾取枪支，然后用鼠标、手柄或触摸式设备上的虚拟摇杆来发射球形子弹。子弹会撞飞关卡中的物理对象，还能从竞技场墙壁上弹开。")
+LocalizedDescriptions=(Language="zh-Hans",Text="第一人称模板包含一个第一人称玩家角色（通过一双手臂呈现）。你可以通过键盘、手柄或触摸式设备上的虚拟摇杆控制角色在关卡中四处移动。你还可以使用鼠标、手柄或触摸式设备上的虚拟摇杆转动角色视角，观察四周。此外，你还可以让角色拾取枪支，然后用鼠标、手柄或触摸式设备上的虚拟摇杆来发射球形子弹。子弹会撞飞关卡中的物理对象，还能从竞技场墙壁上弹开。")
 
 ; This project uses a thumbnail as the project icon
 bThumbnailAsIcon=true
diff --git a/Templates/TP_FirstPerson/manifest.json b/Templates/TP_FirstPerson/manifest.json
--- a/Templates/TP_FirstPerson/manifest.json
+++ b/Templates/TP_FirstPerson/manifest.json
@@ -31,7 +31,7 @@
 		},
 		{
 			"Language": "zh-Hans",
-			"Text": "第一人称模板包含一个第一人称玩家角色（通过一双手臂呈现）。您可以通过键盘、手柄或触摸式设备上的虚拟摇杆控制角色在关卡中四处移动。你还可以使用鼠标、手柄或触摸式设备上的虚拟摇杆转动角色视角，观察四周。此外，你还可以让角色拾取枪支，然后用鼠标、手柄或触摸式设备上的虚拟摇杆来发射球形子弹。子弹会撞飞关卡中的物理对象，还能从竞技场墙壁上弹开。"
+			"Text": "第一人称模板包含一个第一人称玩家角色（通过一双手臂呈现）。你可以通过键盘、手柄或触摸式设备上的虚拟摇杆控制角色在关卡中四处移动。你还可以使用鼠标、手柄或触摸式设备上的虚拟摇杆转动角色视角，观察四周。此外，你还可以让角色拾取枪支，然后用鼠标、手柄或触摸式设备上的虚拟摇杆来发射球形子弹。子弹会撞飞关卡中的物理对象，还能从竞技场墙壁上弹开。"
 		},
 		{
 			"Language": "ja",
diff --git a/Templates/TP_FirstPersonBP/Config/TemplateDefs.ini b/Templates/TP_FirstPersonBP/Config/TemplateDefs.ini
--- a/Templates/TP_FirstPersonBP/Config/TemplateDefs.ini
+++ b/Templates/TP_FirstPersonBP/Config/TemplateDefs.ini
@@ -12,7 +12,7 @@ LocalizedDisplayNames=(Language="ja",Text="ファーストパーソン")
 LocalizedDescriptions=(Language="ja",Text="ファーストパーソンテンプレートには、一人称視点で表示されるプレイヤーキャラクタ（腕として表示されています）がいます。キーボード、コントローラー、タッチデバイスでのバーチャルジョイスティックでこのキャラクタを移動させることができます。追加して、プレイヤーはマウス、コントローラー、タッチデバイスでのバーチャルジョイスティックで周りを見回すことができます。さらにキャラクタは銃を装備して、マウス、コントローラー、タッチデバイスでのバーチャルジョイスティックを使ってシンプルな球体の弾を発射できます。この弾は物理オブジェクトに影響し、アリーナの壁で反射します。")
 
 LocalizedDisplayNames=(Language="zh-Hans",Text="第一人称游戏")
-LocalizedDescriptions=(Language="zh-Hans",Text="第一人称模板拥有一个第一人称的玩家角色（通过一双手臂呈现）。您可以通过键盘、手柄或触摸式设备上的虚拟摇杆控制角色在关卡中四处移动。此外，您还可以使用鼠标、手柄或触摸式设备上的虚拟摇杆转动角色视角，观察四周。角色还装备了枪支，可使用鼠标、手柄或触摸式设备上的虚拟摇杆发射一颗球形子弹，并影响关卡中的一些物理对象，同时从竞技场的墙壁上弹开。")
+LocalizedDescriptions=(Language="zh-Hans",Text="第一人称模板包含一个第一人称玩家角色（通过一双手臂呈现）。你可以通过键盘、手柄或触摸式设备上的虚拟摇杆控制角色在关卡中四处移动。你还可以使用鼠标、手柄或触摸式设备上的虚拟摇杆转动角色视角，观察四周。此外，你还可以让角色拾取枪支，然后用鼠标、手柄或触摸式设备上的虚拟摇杆来发射球形子弹。子弹会撞飞关卡中的物理对象，还能从竞技场墙壁上弹开。")
 
 ; This project uses a thumbnail as the project icon
 bThumbnailAsIcon=true
diff --git a/Templates/TP_FirstPersonBP/manifest.json b/Templates/TP_FirstPersonBP/manifest.json
--- a/Templates/TP_FirstPersonBP/manifest.json
+++ b/Templates/TP_FirstPersonBP/manifest.json
@@ -31,7 +31,7 @@
 		},
 		{
 			"Language": "zh-Hans",
-			"Text": "第一人称模板拥有一个第一人称的玩家角色（通过一双手臂呈现）。您可以通过键盘、手柄或触摸式设备上的虚拟摇杆控制角色在关卡中四处移动。此外，您还可以使用鼠标、手柄或触摸式设备上的虚拟摇杆转动角色视角，观察四周。角色还装备了枪支，可使用鼠标、手柄或触摸式设备上的虚拟摇杆发射一颗球形子弹，并影响关卡中的一些物理对象，同时从竞技场的墙壁上弹开。"
+			"Text": "第一人称模板包含一个第一人称玩家角色（通过一双手臂呈现）。你可以通过键盘、手柄或触摸式设备上的虚拟摇杆控制角色在关卡中四处移动。你还可以使用鼠标、手柄或触摸式设备上的虚拟摇杆转动角色视角，观察四周。此外，你还可以让角色拾取枪支，然后用鼠标、手柄或触摸式设备上的虚拟摇杆来发射球形子弹。子弹会撞飞关卡中的物理对象，还能从竞技场墙壁上弹开。"
 		},
 		{
 			"Language": "ja",
