Update templates to use Enhanced Input as their base input classes. This has no behavioral effect other than stopping the "Automatic Upgrade" from happening every time you boot the editor.

#rnx

diff --git a/Templates/TP_FirstPerson/Config/DefaultInput.ini b/Templates/TP_FirstPerson/Config/DefaultInput.ini
--- a/Templates/TP_FirstPerson/Config/DefaultInput.ini
+++ b/Templates/TP_FirstPerson/Config/DefaultInput.ini
@@ -116,8 +116,8 @@ DoubleClickTime=0.200000
 +AxisMappings=(AxisName="Turn Right / Left Mouse",Scale=1.000000,Key=MouseX)
 +AxisMappings=(AxisName="Look Up / Down Gamepad",Scale=1.000000,Key=Gamepad_RightY)
 +AxisMappings=(AxisName="Look Up / Down Mouse",Scale=-1.000000,Key=MouseY)
-DefaultPlayerInputClass=/Script/Engine.PlayerInput
-DefaultInputComponentClass=/Script/Engine.InputComponent
+DefaultPlayerInputClass=/Script/EnhancedInput.EnhancedPlayerInput
+DefaultInputComponentClass=/Script/EnhancedInput.EnhancedInputComponent
 DefaultTouchInterface=/Engine/MobileResources/HUD/DefaultVirtualJoysticks.DefaultVirtualJoysticks
 -ConsoleKeys=Tilde
 +ConsoleKeys=Tilde
diff --git a/Templates/TP_Puzzle/Config/DefaultInput.ini b/Templates/TP_Puzzle/Config/DefaultInput.ini
--- a/Templates/TP_Puzzle/Config/DefaultInput.ini
+++ b/Templates/TP_Puzzle/Config/DefaultInput.ini
@@ -108,5 +108,6 @@ DoubleClickTime=0.200000
 +ActionMappings=(ActionName="ResetVR",bShift=False,bCtrl=False,bAlt=False,bCmd=False,Key=ValveIndex_Left_Thumbstick_Click)
 DefaultTouchInterface=None
 +ConsoleKeys=Tilde
-
+DefaultPlayerInputClass=/Script/EnhancedInput.EnhancedPlayerInput
+DefaultInputComponentClass=/Script/EnhancedInput.EnhancedInputComponent
 
diff --git a/Templates/TP_ThirdPerson/Config/DefaultInput.ini b/Templates/TP_ThirdPerson/Config/DefaultInput.ini
--- a/Templates/TP_ThirdPerson/Config/DefaultInput.ini
+++ b/Templates/TP_ThirdPerson/Config/DefaultInput.ini
@@ -39,5 +39,6 @@ DoubleClickTime=0.200000
 +AxisMappings=(AxisName="Look Up / Down Mouse",Scale=-1.000000,Key=MouseY)
 DefaultTouchInterface=/Engine/MobileResources/HUD/DefaultVirtualJoysticks.DefaultVirtualJoysticks
 +ConsoleKeys=Tilde
-
+DefaultPlayerInputClass=/Script/EnhancedInput.EnhancedPlayerInput
+DefaultInputComponentClass=/Script/EnhancedInput.EnhancedInputComponent
 
diff --git a/Templates/TP_TopDown/Config/DefaultInput.ini b/Templates/TP_TopDown/Config/DefaultInput.ini
--- a/Templates/TP_TopDown/Config/DefaultInput.ini
+++ b/Templates/TP_TopDown/Config/DefaultInput.ini
@@ -34,8 +34,8 @@ DefaultViewportMouseLockMode=LockOnCapture
 FOVScale=0.011110
 DoubleClickTime=0.200000
 +ActionMappings=(ActionName="SetDestination",bShift=False,bCtrl=False,bAlt=False,bCmd=False,Key=LeftMouseButton)
-DefaultPlayerInputClass=/Script/Engine.PlayerInput
-DefaultInputComponentClass=/Script/Engine.InputComponent
+DefaultPlayerInputClass=/Script/EnhancedInput.EnhancedPlayerInput
+DefaultInputComponentClass=/Script/EnhancedInput.EnhancedInputComponent
 DefaultTouchInterface=None
 -ConsoleKeys=Tilde
 +ConsoleKeys=Tilde
diff --git a/Templates/TP_VehicleAdv/Config/DefaultInput.ini b/Templates/TP_VehicleAdv/Config/DefaultInput.ini
--- a/Templates/TP_VehicleAdv/Config/DefaultInput.ini
+++ b/Templates/TP_VehicleAdv/Config/DefaultInput.ini
@@ -145,4 +145,5 @@ DoubleClickTime=0.200000
 +AxisMappings=(AxisName="MoveRight",Scale=1.000000,Key=ValveIndex_Right_Thumbstick_X)
 DefaultTouchInterface=/Game/Vehicles/Vehicle/LeftVirtualJoystickOnly.LeftVirtualJoystickOnly
 +ConsoleKeys=Tilde
-
+DefaultPlayerInputClass=/Script/EnhancedInput.EnhancedPlayerInput
+DefaultInputComponentClass=/Script/EnhancedInput.EnhancedInputComponent
