[Templates] FP Template.  DefaultEngine.ini didnt get updated and was missing +SM6 and ExtendDefaultLuminanceRange

#content only

diff --git a/Templates/TP_FirstPerson/Config/DefaultEngine.ini b/Templates/TP_FirstPerson/Config/DefaultEngine.ini
--- a/Templates/TP_FirstPerson/Config/DefaultEngine.ini
+++ b/Templates/TP_FirstPerson/Config/DefaultEngine.ini
@@ -21,4 +21,36 @@ r.ReflectionMethod=1
 r.GenerateMeshDistanceFields=True
 r.DynamicGlobalIlluminationMethod=1
 r.Mobile.EnableNoPrecomputedLightingCSMShader=1
+r.DefaultFeature.AutoExposure.ExtendDefaultLuminanceRange=true
+
+
+[/Script/WindowsTargetPlatform.WindowsTargetSettings]
+DefaultGraphicsRHI=DefaultGraphicsRHI_Default
+-D3D12TargetedShaderFormats=PCD3D_SM5
++D3D12TargetedShaderFormats=PCD3D_SM6
+-D3D11TargetedShaderFormats=PCD3D_SM5
++D3D11TargetedShaderFormats=PCD3D_SM5
+Compiler=Default
+AudioSampleRate=48000
+AudioCallbackBufferFrameSize=1024
+AudioNumBuffersToEnqueue=1
+AudioMaxChannels=0
+AudioNumSourceWorkers=4
+SpatializationPlugin=
+SourceDataOverridePlugin=
+ReverbPlugin=
+OcclusionPlugin=
+CompressionOverrides=(bOverrideCompressionTimes=False,DurationThreshold=5.000000,MaxNumRandomBranches=0,SoundCueQualityIndex=0)
+CacheSizeKB=65536
+MaxChunkSizeOverrideKB=0
+bResampleForDevice=False
+MaxSampleRate=48000.000000
+HighSampleRate=32000.000000
+MedSampleRate=24000.000000
+LowSampleRate=12000.000000
+MinSampleRate=8000.000000
+CompressionQualityModifier=1.000000
+AutoStreamingThreshold=0.000000
+SoundCueCookQualityIndex=-1
+
 
