Fixing an issue with Input Mapping Contexts assigned in networked multiplayer games causing an Accessed None. The templates still do not fully support multiplayer, and we're currently discussing if we want to provide this support.

diff --git a/Templates/TP_FirstPerson/Config/DefaultEngine.ini b/Templates/TP_FirstPerson/Config/DefaultEngine.ini
--- a/Templates/TP_FirstPerson/Config/DefaultEngine.ini
+++ b/Templates/TP_FirstPerson/Config/DefaultEngine.ini
@@ -13,7 +13,7 @@ ThreePlayerSplitscreenLayout=FavorTop
 GameInstanceClass=/Script/Engine.GameInstance
 GameDefaultMap=/Game/FirstPerson/Maps/FirstPersonMap.FirstPersonMap
 ServerDefaultMap=/Engine/Maps/Entry
-GlobalDefaultGameMode=/Script/TP_FirstPerson.TP_FirstPersonGameMode
+GlobalDefaultGameMode=/Game/FirstPerson/Blueprints/BP_FirstPersonGameMode.BP_FirstPersonGameMode_C
 GlobalDefaultServerGameMode=None
 
 [/Script/Engine.RendererSettings]
@@ -53,4 +53,17 @@ CompressionQualityModifier=1.000000
 AutoStreamingThreshold=0.000000
 SoundCueCookQualityIndex=-1
 
+[/Script/AndroidFileServerEditor.AndroidFileServerRuntimeSettings]
+bEnablePlugin=True
+bAllowNetworkConnection=True
+SecurityToken=73D7EAB348128FC7E708A79C30473B10
+bIncludeInShipping=False
+bAllowExternalStartInShipping=False
+bCompileAFSProject=False
+bUseCompression=False
+bLogFiles=False
+bReportStats=False
+ConnectionType=USBOnly
+bUseManualIPAddress=False
+ManualIPAddress=
 
diff --git a/Templates/TP_FirstPerson/Source/TP_FirstPerson/TP_FirstPersonPlayerController.cpp b/Templates/TP_FirstPerson/Source/TP_FirstPerson/TP_FirstPersonPlayerController.cpp
new file mode 100644
--- /dev/null
+++ b/Templates/TP_FirstPerson/Source/TP_FirstPerson/TP_FirstPersonPlayerController.cpp
@@ -0,0 +1,19 @@
+// Copyright Epic Games, Inc. All Rights Reserved.
+
+
+#include "TP_FirstPersonPlayerController.h"
+#include "EnhancedInputSubsystems.h"
+
+void ATP_FirstPersonPlayerController::BeginPlay()
+{
+	Super::BeginPlay();
+
+	// get the enhanced input subsystem
+	if (UEnhancedInputLocalPlayerSubsystem* Subsystem = ULocalPlayer::GetSubsystem<UEnhancedInputLocalPlayerSubsystem>(GetLocalPlayer()))
+	{
+		// add the mapping context so we get controls
+		Subsystem->AddMappingContext(InputMappingContext, 0);
+
+		UE_LOG(LogTemp, Warning, TEXT("BeginPlay"));
+	}
+}
\ No newline at end of file
diff --git a/Templates/TP_FirstPerson/Source/TP_FirstPerson/TP_FirstPersonPlayerController.h b/Templates/TP_FirstPerson/Source/TP_FirstPerson/TP_FirstPersonPlayerController.h
new file mode 100644
--- /dev/null
+++ b/Templates/TP_FirstPerson/Source/TP_FirstPerson/TP_FirstPersonPlayerController.h
@@ -0,0 +1,31 @@
+// Copyright Epic Games, Inc. All Rights Reserved.
+
+#pragma once
+
+#include "CoreMinimal.h"
+#include "GameFramework/PlayerController.h"
+#include "TP_FirstPersonPlayerController.generated.h"
+
+class UInputMappingContext;
+
+/**
+ *
+ */
+UCLASS()
+class TP_FIRSTPERSON_API ATP_FirstPersonPlayerController : public APlayerController
+{
+	GENERATED_BODY()
+	
+protected:
+
+	/** Input Mapping Context to be used for player input */
+	UPROPERTY(EditAnywhere, BlueprintReadOnly, Category = Input)
+	UInputMappingContext* InputMappingContext;
+
+	// Begin Actor interface
+protected:
+
+	virtual void BeginPlay() override;
+
+	// End Actor interface
+};
