TP_FirstPersonBP - Fixed reference to variant AIController class on base template character.
TP_FirstPerson, BP - Fixed third person tick flag that was preventing pose updates on some platforms.
TP_FirstPerson - Matched camera clip distance to prevent shooter first person guns from being clipped.
#rnx

diff --git a/Templates/TP_FirstPerson/Config/DefaultEngine.ini b/Templates/TP_FirstPerson/Config/DefaultEngine.ini
--- a/Templates/TP_FirstPerson/Config/DefaultEngine.ini
+++ b/Templates/TP_FirstPerson/Config/DefaultEngine.ini
@@ -60,3 +60,6 @@ SoundCueCookQualityIndex=-1
 [/Script/AIModule.AISystem]
 bForgetStaleActors=True
 
+[/Script/Engine.Engine]
+NearClipPlane=5.000000
+
