TP_FirstPerson
- Removed unused reference to legacy mobile controls interface from DefaultInput
#rnx

diff --git a/Templates/TP_FirstPerson/Config/DefaultInput.ini b/Templates/TP_FirstPerson/Config/DefaultInput.ini
--- a/Templates/TP_FirstPerson/Config/DefaultInput.ini
+++ b/Templates/TP_FirstPerson/Config/DefaultInput.ini
@@ -103,7 +103,6 @@ FOVScale=0.011110
 DoubleClickTime=0.200000
 DefaultPlayerInputClass=/Script/EnhancedInput.EnhancedPlayerInput
 DefaultInputComponentClass=/Script/EnhancedInput.EnhancedInputComponent
-DefaultTouchInterface=/Game/FirstPerson/Input/MobileControls.MobileControls
 -ConsoleKeys=Tilde
 +ConsoleKeys=Tilde
 
